Introducing Code Quality Test plugin.
Extension of the Unreal Engine FAutomationTestBase to provide test fixtures and common automation testing commands.
Add the before/after test abilities, while resetting state between tests automatically.

PR #10042: CQ Test Plugin (Contributed by [REDACTED_PERSON])

diff --git a/Engine/Plugins/Tests/CQTest/CQTest.uplugin b/Engine/Plugins/Tests/CQTest/CQTest.uplugin
new file mode 100644
--- /dev/null
+++ b/Engine/Plugins/Tests/CQTest/CQTest.uplugin
@@ -0,0 +1,28 @@
+{
+  "FileVersion": 3,
+  "Version": 1,
+  "VersionName": "1.0",
+  "FriendlyName": "Code Quality Unreal Test Plugin",
+  "Description": "Simplified testing classes for Unreal Engine",
+  "Category": "Code Quality",
+  "CreatedBy": "Studios Quality - Code Quality",
+  "CreatedByURL": "",
+  "DocsURL": "",
+  "MarketplaceURL": "",
+  "SupportURL": "",
+  "CanContainContent": false,
+  "IsBetaVersion": false,
+  "Installed": false,
+  "Modules": [
+		{
+			"Name": "CQTest",
+			"Type": "DeveloperTool",
+			"LoadingPhase": "Default"
+		},
+		{
+			"Name": "CQTestTests",
+			"Type": "DeveloperTool",
+			"LoadingPhase": "Default"
+		}
+	]
+}
\ No newline at end of file
diff --git a/Engine/Plugins/Tests/CQTest/README.md b/Engine/Plugins/Tests/CQTest/README.md
new file mode 100644
--- /dev/null
+++ b/Engine/Plugins/Tests/CQTest/README.md
@@ -0,0 +1,363 @@
+# Introduction 
+Extension of the Unreal Engine FAutomationTestBase to provide test fixtures and common automation testing commands.
+
+# Why CQTest?
+There are other valid ways of testing in Unreal engine.  One option is to use the provided macros from Unreal Engine: [docs](https://docs.unrealengine.com/5.1/en-US/automation-technical-guide/)
+```cpp
+    IMPLEMENT_SIMPLE_AUTOMATION_TEST(FMinimalTest, "Game.Test", EAutomationTestFlags::EditorContext | EAutomationTestFlags::EngineFilter)
+	bool FMinmalTest::RunTest(const FString& Parameters) 
+	{
+		TestTrue(TEXT("True should be true"), true);
+		return true;
+	}
+```
+
+Unreal has also developed their [spec test framework](https://docs.unrealengine.com/4.27/en-US/TestingAndOptimization/Automation/AutomationSpec/), which is inspired by Behavior Driven Design
+```cpp
+    DEFINE_SPEC(FMinimalTest, "Game.Test", EAutomationTestFlags::ProductFilter | EAutomationTestFlags::ApplicationContextMask)
+	void FMinimalTest::Define() 
+	{
+		Describe("Assertions", [this]() 
+		{
+			It("Should pass when testing that true is true", [this]() 
+			{
+				TestTrue(TEXT("True should be true"), true);
+			});
+		});
+	}
+```
+With the spec tests, be careful about capturing state
+
+```cpp
+    BEGIN_DEFINE_SPEC(FMinimalTest, "Game.Test", EAutomationTestFlags::ProductFilter | EAutomationTestFlags::ApplicationContextMask)
+		uint32 SomeValue = 3;
+	END_DEFINE_SPEC
+
+	void FMinimalTest::Define() 
+	{
+		Describe("Assertions", [this]() 
+		{
+			uint describeValue = 42;
+			It("Has access to members defined on the spec", [this]() 
+			{
+				TestEqual(TEXT("Class value should be set"), SomeValue, 3);
+			});
+			
+			xIt("Does not capture variables described inside of lambdas", [this]() 
+			{
+				TestEqual(TEXT("DescribeValue will now be garbage as it went out of scope"), describeValue, 42);
+			});
+		});
+	}
+```
+
+The inspiration for CQTest was to add the before/after test abilities, while resetting state between tests automatically.  One of the guiding principles is to make easy things easy.
+
+```cpp
+	TEST(MinimalTest, "Game.Test") 
+	{
+		ASSERT_THAT(IsTrue(true));
+	}
+	
+	TEST_CLASS(MinimalFixture, "Game.Test") 
+	{
+		uint32 SomeNumber = 0;
+		BEFORE_EACH() 
+		{
+			SomeNumber++;
+		}
+
+		TEST_METHOD(MinimalFixture, CanAccessMembers) 
+		{
+			ASSERT_THAT(AreEqual(1, SomeNumber));  // passes every time
+		}
+	};
+```
+
+# Installation
+Clone the project, into your project's Plugin folder.
+Inside the project that you want to test, you'll need to change 2 things:
+In the .uproject file of the project you want to test, add the following to the Plugins section
+
+```json
+		{
+			"Name": "CQTest",
+			"Enabled": true
+		}
+```
+
+Then in the project's .Build.cs file, you'll want to add the following to the PrivateDependencyModuleNames.  Something like
+
+```csharp
+		PrivateDependencyModuleNames.AddRange(
+			new string[] {
+					"Core",
+					"CoreUObject",
+					"Engine",
+					"CQTest"  <----
+				 }
+			);
+```
+
+	
+# Build
+No special steps are required to build the plugin project, it should build with the rest of the project.
+
+# Test
+This plugin has a set of tests to validate and document the behavior.  To run tests in Unreal
+- Launch the editor
+- Find the Tools drop down and select Session Frontend
+- Navigate to the Automation tab
+- By default, the tests should be listed first under " Product.Plugins.CQTest"
+- Select the tests you would like to run and press 'Start Tests'
+
+# Examples
+
+Tests can be as simple as
+
+```cpp
+    #include "CQTest.h"
+	
+	TEST(MinimalTest, "Game.MyGame") 
+	{
+		ASSERT_THAT(IsTrue(true));
+	}
+```
+
+For setup and teardown, or common state between multiple tests, or to group related tests, use the TEST_CLASS macro.
+
+```cpp
+    #include "CQTest.h"
+	TEST_CLASS(MyNeatTest, "Game.MyGame") 
+	{
+		bool SetupCalled = false;
+		uint32 SomeNumber = 0;
+		Thing* Thing = nullptr;
+		
+		BEFORE_EACH() 
+		{
+			SetupCalled = true;
+			SomeNumber++;
+			Thing = new Thing();
+		}
+		
+		AFTER_EACH() 
+		{
+			delete Thing; //Should normally use RAII for things like this
+		}
+		
+	protected:
+		bool UsefulHelperMethod() const 
+		{
+			return true; 
+		}
+		
+		TEST_METHOD(BeforeRunTest_CallsSetup) 
+		{
+			ASSERT_THAT(IsTrue(SetupCalled));
+		}
+		
+		TEST_METHOD(ProtectedMembers_AreAccessible) 
+		{
+			ASSERT_THAT(IsTrue(UsefulHelperMethod()));
+		}
+		
+		TEST_METHOD(DataMembers_BetweenTestRuns_AreReset) 
+		{
+			ASSERT_THAT(AreEqual(1, SomeNumber));
+		}
+	};
+```
+
+Test Directory determines where in the Automation tab the tests appear.  In the example above, we specify "Game.MyGame", but you may also have an auto-generated test directory based on the folder structure.
+
+```cpp
+	TEST_CLASS(MyNeatTest, GenerateTestDirectory)
+	{
+	};
+
+	TEST_CLASS(MyNeatTest, "Game.Test.[GenerateTestDirectory].Validation")
+	{
+	};
+```
+
+Constructors (and destructors) are available.  Destructors shouldn't throw, and you shouldn't put assertions in them (as they are called after the testing framework is done with the test).
+
+```cpp
+	TEST_CLASS(SomeTestClass, "Game.Test")
+	{
+		bool constructed = false;
+		SomeTestClass()
+			: constructed(true)
+		{
+		}
+		
+		TEST_METHOD(ConstructorIsCalled)
+		{
+			ASSERT_THAT(IsTrue(constructed));
+		}
+	};
+```	
+
+Latent actions are supported with the TEST_CLASS macro.  Each step will complete all latent actions before moving to the next
+
+```cpp
+    TEST_CLASS(LatentActionTest, "Game.Test") 
+	{
+		uint32 calls = 0;
+		BEFORE_EACH() 
+		{
+			AddCommand(new FExecute([&]() { calls++; }));
+		}
+		
+		AFTER_EACH() 
+		{
+			AddCommand(new FExecute([&]() { calls++; })); // executed after the next line, as it is a latent action
+			ASSERT_THAT(AreEqual(2, calls));
+		}
+		
+		TEST_METHOD(PerformLatentAction) 
+		{
+			ASSERT_THAT(AreEqual(1, calls));
+			AddCommand(new FExecute([&]() { calls++; }));
+		}
+	};
+```
+
+Also available for commands is a fluent command builder
+
+```cpp
+	TEST_METHOD(SomeTest) 
+	{
+		TestCommandBuilder
+			.Do([&]() { StepOne(); })
+			.Then([&]() { StepTwo(); })
+			.Until([&]() { return StepThreeComplete(); })
+			.Then([&]() { ASSERT_THAT(IsTrue(SomethingImportant)); });
+	}
+```
+
+The framework will ensure that all of those commands happen in order using a future pattern.
+Similarly, the framework will ensure that a test can await a ticking object.  See GameObjectsTickTest for an example
+
+# Extending the framework
+The framework has been designed to allow for extensions in a couple areas.  See ExtensionTests.cpp for in-code examples.
+
+## Assertions
+Not all platforms support exceptions, and so the assertions are unable to rely on them.
+There are a few options here:
+  We could just throw exceptions, and only run tests on platforms which support exceptions
+  We could return a [[nodiscard]] bool to encourage checking each assertion and returning if it fails
+  We could return a normal bool and rely on people to check it when it's important.
+Exceptions have the advantage of working in helper functions and lambdas, as well as not depending on human diligence.
+A normal bool is less noisy, and allows developers to use intellisense, but is more error prone
+The default implementation used is the [[nodiscard]] bool, with a helper macro ASSERT_THAT which does the early return check for you.
+
+You can use your own types within the Assert.AreEqual and Assert.AreNotEqual methods assuming you have the == and != operators defined as needed.
+In addition, the error message will print out the string version of your type, assuming you have a ToString method defined as well.  The framework will complain if it doesn't know how to print your value.
+You can find examples of providing a string to the framework in CQTestConvertTests.cpp, but below is a simple example.
+```cpp
+struct MyCustomType
+{
+	int32 Value;
+	bool operator==(const MyCustomType& other) const
+	{
+		return Value == other.Value;
+	}
+	bool operator!=(const MyCustomType& other) const
+	{
+		return !(*this == other);
+	}
+	
+	FString ToString() const
+	{
+		//your to string logic
+		return FString();
+	}
+};
+enum struct MyCustomEnum
+{
+	Red, Green, Blue
+};
+template<>
+FString CQTestConvert::ToString(const MyCustomEnum&)
+{
+	//your to string logic
+	return FString();
+}
+```
+
+You are able to customize the assertions which are available, and how they behave.
+See CQTestTests/Private/ExtensionTests.cpp for an example
+Below is some untested example code to inspire ideas
+
+```cpp
+	struct FluentAsserter : public ThrowingAsserter
+	{
+	private:
+		int CurrentIntValue = 0;
+		TArray<FString> Errors;
+		
+	public:
+		FluentAsserter(FAutomationTestBase& testRunner)
+			: ThrowingAsserter(testRunner)
+		{
+		}
+		
+		~FluentAsserter()
+		{
+			for(const auto& error : Errors)
+			{
+				TestRunner.AddError(error);
+			}
+		}
+		
+		FluentAsserter& That(int value)
+		{
+			CurrentIntValue = value;
+			return *this;
+		}
+		
+		FluentAsserter& Equals(int value)
+		{
+			if(CurrentIntValue != value)
+			{
+				Errors.Add(FString::Printf("%d != %d", CurrentIntValue, value));
+			}
+			return *this;
+		}
+	};
+```
+
+From here, you could create macros your studio uses to create tests
+
+```cpp
+	#define MY_STUDIO_TEST_CLASS(_ClassName, _TestDir) TEST_CLASS_WITH_ASSERTS(_ClassName, _TestDir, FluentAsserter)
+	#define MY_STUDIO_TEST(_TestName, _TestDir) \
+	MY_STUDIO_TEST_CLASS(_TestName, _TestDir) \
+	{ \
+		TEST_METHOD(_TestName##_Method); \
+	};\
+	void _TestName::_TestName##_Method()
+```
+
+## Base test class
+Similarly there may be a use case to create many tests which have the same member variables or helper methods.  This can be implemented by extending the test class
+
+```cpp
+	template<typename Derived, typename AsserterType>
+	struct ActorTest : public Test<Derived, AsserterType>
+	{
+		SpawnHelper Spawner;
+	};
+```
+
+And creating a macro which uses it
+
+```cpp
+	#define ACTOR_TEST(_ClassName, _TestDir) TEST_CLASS_WITH_BASE(_ClassName, _TestDir, ActorTest)
+```
+
+# Contribute
+Improvements like bug fixes and extensions are welcome when accompanied by unit tests.
+Feel free to reach out to [Dustin Randall](mailto:[REDACTED_EMAIL]) with questions or comments
diff --git a/Engine/Plugins/Tests/CQTest/Source/CQTest/CQTest.Build.cs b/Engine/Plugins/Tests/CQTest/Source/CQTest/CQTest.Build.cs
new file mode 100644
--- /dev/null
+++ b/Engine/Plugins/Tests/CQTest/Source/CQTest/CQTest.Build.cs
@@ -0,0 +1,30 @@
+// Copyright Epic Games, Inc. All Rights Reserved.
+
+using UnrealBuildTool;
+
+public class CQTest : ModuleRules
+{
+	public CQTest(ReadOnlyTargetRules Target)
+		: base(Target)
+	{
+        PublicDependencyModuleNames.Add("DeveloperSettings");
+
+		PrivateDependencyModuleNames.AddRange(
+			new string[] {
+					"Core",
+					"CoreUObject",
+					"Engine"
+				 }
+			);
+
+		if (Target.bBuildEditor)
+		{
+			PrivateDependencyModuleNames.AddRange(
+				new string[] {
+					"EngineSettings",
+					"LevelEditor",
+					"UnrealEd"
+			});
+		}
+	}
+}
diff --git a/Engine/Plugins/Tests/CQTest/Source/CQTest/CQTestModule.cpp b/Engine/Plugins/Tests/CQTest/Source/CQTest/CQTestModule.cpp
new file mode 100644
--- /dev/null
+++ b/Engine/Plugins/Tests/CQTest/Source/CQTest/CQTestModule.cpp
@@ -0,0 +1,5 @@
+// Copyright Epic Games, Inc. All Rights Reserved.
+
+#include "Modules/ModuleManager.h"
+
+IMPLEMENT_MODULE( FDefaultModuleImpl, CQTest )
diff --git a/Engine/Plugins/Tests/CQTest/Source/CQTest/Private/ActorTestSpawner.cpp b/Engine/Plugins/Tests/CQTest/Source/CQTest/Private/ActorTestSpawner.cpp
new file mode 100644
--- /dev/null
+++ b/Engine/Plugins/Tests/CQTest/Source/CQTest/Private/ActorTestSpawner.cpp
@@ -0,0 +1,98 @@
+// Copyright Epic Games, Inc. All Rights Reserved.
+
+#include "ActorTestSpawner.h"
+#include "TestGameInstance.h"
+
+#include "EngineUtils.h"
+#include "Engine/Engine.h"
+
+namespace
+{
+	static constexpr bool bTriggerWorldDestroyCallbacks = true;
+
+	void DestroyWorld(UWorld* GameWorld)
+	{
+		if (GameWorld->AreActorsInitialized())
+		{
+			for (AActor* const Actor : FActorRange(GameWorld))
+			{
+				if (Actor)
+				{
+					Actor->RouteEndPlay(EEndPlayReason::LevelTransition);
+				}
+			}
+		}
+
+		GEngine->ShutdownWorldNetDriver(GameWorld);
+		GameWorld->DestroyWorld(bTriggerWorldDestroyCallbacks);
+		GameWorld->SetPhysicsScene(nullptr);
+		GEngine->DestroyWorldContext(GameWorld);
+	}
+
+	void DestroySpawnedActors(TArray<TWeakObjectPtr<AActor>>& SpawnedActors)
+	{
+		// Destroy actors in reverse order of creation
+		for (int32 Index = SpawnedActors.Num() - 1; Index >= 0; Index--)
+		{
+			if (AActor* CastActor = SpawnedActors[Index].Get())
+			{
+				CastActor->Destroy(true);
+			}
+		}
+
+		SpawnedActors.Empty();
+	}
+
+	void DestroySpawnedObjects(TArray<TWeakObjectPtr<UObject>>& InSpawnedObjects)
+	{
+		// Destroy objects in reverse order of creation
+		for (int32 Index = InSpawnedObjects.Num() - 1; Index >= 0; Index--)
+		{
+			if (UObject* Obj = InSpawnedObjects[Index].Get())
+			{
+				Obj->BeginDestroy();
+			}
+		}
+
+		InSpawnedObjects.Empty();
+	}
+} // namespace
+
+FActorTestSpawner::~FActorTestSpawner() {
+	if (GameInstance)
+	{
+		GameInstance->Shutdown();
+		GameInstance->RemoveFromRoot();
+		GameInstance = nullptr;
+	}
+
+	if (GameWorld)
+	{
+		DestroySpawnedActors(SpawnedActors);
+		DestroySpawnedObjects(SpawnedObjects);
+		DestroyWorld(GameWorld);
+		GameWorld->RemoveFromRoot();
+		GameWorld = nullptr;
+	}
+}
+
+UWorld* FActorTestSpawner::CreateWorld()
+{
+	FWorldContext& WorldContext = GEngine->CreateNewWorldContext(EWorldType::Game);
+	UWorld* Result = UWorld::CreateWorld(EWorldType::Game, false, NAME_None);
+	check(Result != nullptr);
+	Result->AddToRoot();
+	WorldContext.SetCurrentWorld(Result);
+
+	Result->InitializeActorsForPlay(FURL());
+	Result->CreatePhysicsScene();
+
+	return Result;
+}
+
+void FActorTestSpawner::InitializeGameSubsystems()
+{
+	GameInstance = NewObject<UTestGameInstance>();
+	GameInstance->AddToRoot();
+	GameInstance->InitializeForTest(GetWorld());
+}
diff --git a/Engine/Plugins/Tests/CQTest/Source/CQTest/Private/CQTestBlueprintHelper.cpp b/Engine/Plugins/Tests/CQTest/Source/CQTest/Private/CQTestBlueprintHelper.cpp
new file mode 100644
--- /dev/null
+++ b/Engine/Plugins/Tests/CQTest/Source/CQTest/Private/CQTestBlueprintHelper.cpp
@@ -0,0 +1,47 @@
+// Copyright Epic Games, Inc. All Rights Reserved.
+
+#include "CQTestBlueprintHelper.h"
+
+#include "Misc/Paths.h"
+#include "Engine/Blueprint.h"
+#include "AssetRegistry/AssetData.h"
+#include "AssetRegistry/AssetRegistryModule.h"
+
+UClass* FCQTestBlueprintHelper::GetBlueprintClass(const FString& Directory, const FString& Name)
+{
+	const FString FileName = FString::Printf(TEXT("%s.%s"), *Name, *Name);
+	const FString Path = FPaths::Combine(*Directory, FileName);
+	UObject* BpObject = Cast<UObject>(StaticLoadObject(UObject::StaticClass(), nullptr, *Path));
+	check(BpObject);
+
+	UBlueprint* AsBlueprint = Cast<UBlueprint>(BpObject);
+	check(AsBlueprint);
+
+	return AsBlueprint->GeneratedClass;
+}
+
+UObject* FCQTestBlueprintHelper::FindDataBlueprint(const FString& Directory, const FString& Name)
+{
+	static TMap<FString, FAssetData> BlueprintCache{};
+	static TSet<FString> LoadedDirectories{};
+	if (!LoadedDirectories.Contains(Directory))
+	{
+		TArray<FAssetData> AssetData;
+		const FAssetRegistryModule& AssetRegistry = FModuleManager::LoadModuleChecked<FAssetRegistryModule>("AssetRegistry");
+
+		AssetRegistry.Get().GetAssetsByPath(FName(Directory), AssetData, true);
+		for (const auto& Data : AssetData)
+		{
+			BlueprintCache.Add(Data.GetAsset()->GetFName().ToString(), Data);
+		}
+	
+		LoadedDirectories.Add(Directory);
+	}
+
+	if (FAssetData* Data = BlueprintCache.Find(Name))
+	{
+		return Data->GetAsset();
+	}
+
+	return nullptr;
+}
\ No newline at end of file
diff --git a/Engine/Plugins/Tests/CQTest/Source/CQTest/Private/MapTestSpawner.cpp b/Engine/Plugins/Tests/CQTest/Source/CQTest/Private/MapTestSpawner.cpp
new file mode 100644
--- /dev/null
+++ b/Engine/Plugins/Tests/CQTest/Source/CQTest/Private/MapTestSpawner.cpp
@@ -0,0 +1,42 @@
+// Copyright Epic Games, Inc. All Rights Reserved.
+
+#include "MapTestSpawner.h"
+#include "Commands/TestCommands.h"
+
+#include "Tests/AutomationCommon.h"
+#include "GameFramework/PlayerController.h"
+#include "Engine/Engine.h"
+
+void FMapTestSpawner::AddWaitUntilLoadedCommand(FAutomationTestBase* TestRunner)
+{
+	check(PieWorld == nullptr);
+
+	const FString FileName = FString::Printf(TEXT("%s.%s"), *MapName, *MapName);
+	const FString Path = FPaths::Combine(MapDirectory, FileName);
+	bool bOpened = AutomationOpenMap(Path);
+	check(bOpened);
+
+	ADD_LATENT_AUTOMATION_COMMAND(FWaitUntil(TestRunner, [&]() -> bool {
+		for (const auto& Context : GEngine->GetWorldContexts())
+		{
+			if (((Context.WorldType == EWorldType::PIE) || (Context.WorldType == EWorldType::Game)) && (Context.World() != nullptr))
+			{
+				PieWorld = Context.World();
+				return true;
+			}
+		}
+
+		return false;
+	}));
+}
+
+UWorld* FMapTestSpawner::CreateWorld()
+{
+	checkf(PieWorld, TEXT("Must call AddWaitUntilLoadedCommand in BEFORE_TEST"));
+	return PieWorld;
+}
+
+APawn* FMapTestSpawner::FindFirstPlayerPawn()
+{
+	return GetWorld().GetFirstPlayerController()->GetPawn();
+}
diff --git a/Engine/Plugins/Tests/CQTest/Source/CQTest/Private/SpawnHelper.cpp b/Engine/Plugins/Tests/CQTest/Source/CQTest/Private/SpawnHelper.cpp
new file mode 100644
--- /dev/null
+++ b/Engine/Plugins/Tests/CQTest/Source/CQTest/Private/SpawnHelper.cpp
@@ -0,0 +1,24 @@
+// Copyright Epic Games, Inc. All Rights Reserved.
+
+#include "SpawnHelper.h"
+
+FSpawnHelper::~FSpawnHelper() 
+{
+	if (GameWorld)
+	{
+		GameWorld->RemoveFromRoot();
+		GameWorld = nullptr;
+	}
+}
+
+UWorld& FSpawnHelper::GetWorld()
+{
+	if (!GameWorld)
+	{
+		GameWorld = CreateWorld();
+		check(GameWorld);
+		GameWorld->AddToRoot();
+	}
+
+	return *GameWorld;
+}
\ No newline at end of file
diff --git a/Engine/Plugins/Tests/CQTest/Source/CQTest/Private/TestDirectoryGenerator.cpp b/Engine/Plugins/Tests/CQTest/Source/CQTest/Private/TestDirectoryGenerator.cpp
new file mode 100644
--- /dev/null
+++ b/Engine/Plugins/Tests/CQTest/Source/CQTest/Private/TestDirectoryGenerator.cpp
@@ -0,0 +1,63 @@
+// Copyright Epic Games, Inc. All Rights Reserved.
+#include "Misc/App.h"
+#include "Misc/Paths.h"
+
+namespace TestDirectoryGenerator
+{
+
+	void SplitPath(const FString& Path, FString* Lhs, FString* Rhs, FString* SourceType)
+	{
+		static const FString Plugins = TEXT("/Plugins/");
+		static const FString Platforms = TEXT("/Platforms/");
+		static const FString Source = TEXT("/Source/");
+		static const FString Slash = TEXT("/");
+
+		FString SplitString = Source;
+		if (Path.Find(Plugins) != INDEX_NONE)
+		{
+			Path.Split(Plugins, Lhs, Rhs);
+			*SourceType = TEXT("Plugins");
+		}
+		else if (Path.Find(Platforms) != INDEX_NONE)
+		{
+			Path.Split(Platforms, Lhs, Rhs);
+			*SourceType = TEXT("Platforms");
+		}
+		else if (Path.Find(Source) != INDEX_NONE)
+		{
+			Path.Split(Source, Lhs, Rhs);
+			*SourceType = TEXT("Source");
+			SplitString = Slash;
+		}
+		else
+		{
+			checkf(false, TEXT("%s Should contain Plugins, Platforms, or Source"), *Path);
+			*Lhs = TEXT("Unknown");
+			*Rhs = TEXT("Unknown");
+			*SourceType = TEXT("Unknown");
+			return;
+		}
+
+		Lhs->Split(Slash, nullptr, Lhs, ESearchCase::CaseSensitive, ESearchDir::FromEnd);
+		if (Lhs->IsEmpty())
+		{
+			*Lhs = TEXT("Unknown");
+		}
+		Rhs->Split(SplitString, Rhs, nullptr);
+		if (Rhs->IsEmpty())
+		{
+			*Rhs = TEXT("Unknown");
+		}
+	}
+
+	CQTEST_API FString Generate(const FString& Filename)
+	{
+		static const FString Engine = TEXT("Engine/");
+		const FString Path(Filename.Replace(TEXT("\\"), TEXT("/"), ESearchCase::CaseSensitive).Replace(TEXT("../"), *Engine, ESearchCase::CaseSensitive));
+
+		FString Before, After, SourceType;
+		SplitPath(Path, &Before, &After, &SourceType);
+		After = After.Replace(TEXT("/"), TEXT("."));
+		return FString::Printf(TEXT("%s.%s.%s"), *Before, *SourceType, *After);
+	}
+}
\ No newline at end of file
diff --git a/Engine/Plugins/Tests/CQTest/Source/CQTest/Private/TestGameInstance.cpp b/Engine/Plugins/Tests/CQTest/Source/CQTest/Private/TestGameInstance.cpp
new file mode 100644
--- /dev/null
+++ b/Engine/Plugins/Tests/CQTest/Source/CQTest/Private/TestGameInstance.cpp
@@ -0,0 +1,20 @@
+// Copyright Epic Games, Inc. All Rights Reserved.
+
+#include "TestGameInstance.h"
+#include "Engine/Engine.h"
+
+UTestGameInstance::UTestGameInstance(const FObjectInitializer& InObjectInitializer)
+	: Super(InObjectInitializer)
+{
+}
+
+void UTestGameInstance::InitializeForTest(UWorld& InWorld) {
+	FWorldContext* TestWorldContext = GEngine->GetWorldContextFromWorld(&InWorld);
+	check(TestWorldContext);
+	WorldContext = TestWorldContext;
+
+	WorldContext->OwningGameInstance = this;
+	InWorld.SetGameInstance(this);
+
+	Init();
+}
\ No newline at end of file
diff --git a/Engine/Plugins/Tests/CQTest/Source/CQTest/Public/ActorTestSpawner.h b/Engine/Plugins/Tests/CQTest/Source/CQTest/Public/ActorTestSpawner.h
new file mode 100644
--- /dev/null
+++ b/Engine/Plugins/Tests/CQTest/Source/CQTest/Public/ActorTestSpawner.h
@@ -0,0 +1,22 @@
+// Copyright Epic Games, Inc. All Rights Reserved.
+
+#pragma once
+
+#include "SpawnHelper.h"
+
+class UTestGameInstance;
+
+/// Class for spawning actors in an ActorTest context (no PIE loaded)
+struct CQTEST_API FActorTestSpawner : public FSpawnHelper
+{
+	FActorTestSpawner() = default;
+	virtual ~FActorTestSpawner() override;
+
+	void InitializeGameSubsystems();
+
+protected:
+	virtual UWorld* CreateWorld() override;
+
+private:
+	UTestGameInstance* GameInstance{ nullptr };
+};
\ No newline at end of file
diff --git a/Engine/Plugins/Tests/CQTest/Source/CQTest/Public/Assert/CQTestConvert.h b/Engine/Plugins/Tests/CQTest/Source/CQTest/Public/Assert/CQTestConvert.h
new file mode 100644
--- /dev/null
+++ b/Engine/Plugins/Tests/CQTest/Source/CQTest/Public/Assert/CQTestConvert.h
@@ -0,0 +1,112 @@
+// Copyright Epic Games, Inc. All Rights Reserved.
+#pragma once
+
+#include "Containers/UnrealString.h"
+#include "UObject/NameTypes.h"
+#include <sstream>
+#include <type_traits>
+
+namespace CQTestConvert
+{
+	using PrintStream = std::ostringstream;
+}
+
+namespace
+{
+	template <typename T, typename... Args>
+	class THasToString
+	{
+		template <typename C, typename = decltype(std::declval<C>().ToString(std::declval<Args>()...))>
+		static std::true_type test(int);
+		template<typename C>
+		static std::false_type test(...);
+
+	public:
+		static constexpr bool value = decltype(test<T>(0))::value;
+	};
+
+	template<typename T>
+	class THasOStream
+	{
+		template <typename C, typename = decltype(std::declval<CQTestConvert::PrintStream&>().operator<<(std::declval<C>()))>
+		static std::true_type test(int);
+		template <typename C, typename = decltype(operator<<(std::declval<CQTestConvert::PrintStream&>(), std::declval<C>()))>
+		static std::true_type test(char);
+		template<typename C>
+		static std::false_type test(...);
+
+	public:
+		static constexpr bool value = decltype(test<T>(0))::value;
+	};
+
+	static_assert(THasToString<FName>::value, "FName should have ToString");
+	static_assert(THasOStream<int>::value, "int should have an OStream operator");
+	static_assert(THasOStream<int32>::value, "int32 should have an OStream operator");
+
+	struct SomeTestStruct
+	{
+	};
+
+	static_assert(!THasToString<SomeTestStruct>::value, "Struct without ToString should not have ToString");
+	static_assert(!THasOStream<SomeTestStruct>::value, "Struct without OStream operator should not have OStream");
+} // namespace
+
+
+namespace CQTestConvert
+{
+	template <typename T>
+	FString ToString(const T& Input)
+	{
+		if constexpr (THasToString<T>::value)
+		{
+			return Input.ToString();
+		}
+		else if constexpr (THasOStream<T>::value)
+		{
+			PrintStream stream;
+			stream << Input;
+			return FString(stream.str().c_str(), stream.str().length());
+		}
+		else if constexpr (std::is_enum_v<T>)
+		{
+			return ToString(static_cast<std::underlying_type_t<T>>(Input));
+		}
+		else
+		{
+			ensureMsgf(false, TEXT("Did not find ToString, ostream operator, or CQTestConvert::ToString template specialization for provided type. Cast provided type to something else, or see CqTestConvertTests.cpp for examples"));
+			return TEXT("value");
+		}
+	}
+
+	template <>
+	FString ToString(const FString& Input)
+	{
+		return Input;
+	}
+
+	template<typename TElement, typename TAllocator>
+	FString ToString(const TArray<TElement, TAllocator>& Input)
+	{
+		FString result = FString::JoinBy(Input, TEXT(", "), [](const TElement& e) {
+			return CQTestConvert::ToString(e);
+		});
+		return TEXT("[") + result + TEXT("]");
+	}
+
+	template <typename TElement, typename TAllocator>
+	FString ToString(const TSet<TElement, TAllocator>& Input)
+	{
+		FString result = FString::JoinBy(Input, TEXT(", "), [](const TElement& e) {
+			return CQTestConvert::ToString(e);
+		});
+		return TEXT("[") + result + TEXT("]");
+	}
+
+	template<typename TKey, typename TValue>
+	FString ToString(const TMap<TKey, TValue>& Input)
+	{
+		return FString::JoinBy(Input, TEXT(", "), [](const auto& kvp) {
+			return TEXT("{") + CQTestConvert::ToString(kvp.Key) + TEXT(":") + CQTestConvert::ToString(kvp.Value) + TEXT("}");
+		});
+	}
+}
diff --git a/Engine/Plugins/Tests/CQTest/Source/CQTest/Public/Assert/NoDiscardAsserter.h b/Engine/Plugins/Tests/CQTest/Source/CQTest/Public/Assert/NoDiscardAsserter.h
new file mode 100644
--- /dev/null
+++ b/Engine/Plugins/Tests/CQTest/Source/CQTest/Public/Assert/NoDiscardAsserter.h
@@ -0,0 +1,92 @@
+// Copyright Epic Games, Inc. All Rights Reserved.
+
+#pragma once
+
+#include "Misc/AutomationTest.h"
+
+struct FNoDiscardAsserter
+{
+protected:
+	FAutomationTestBase& TestRunner;
+
+public:
+	explicit FNoDiscardAsserter(FAutomationTestBase& TestRunner)
+		: TestRunner(TestRunner) {}
+
+	void ExpectError(FString Error, int32 Count = 1);
+	void ExpectErrorRegex(FString Error, int32 Count = 1);
+
+	void Fail(FString error);
+
+	[[nodiscard]] bool IsTrue(bool Condition);
+	[[nodiscard]] bool IsTrue(bool Condition, const char* FailureMessage);
+	[[nodiscard]] bool IsTrue(bool Condition, const TCHAR* FailureMessage);
+	[[nodiscard]] bool IsTrue(bool Condition, const FString& FailureMessage);
+
+	[[nodiscard]] bool IsFalse(bool Condition);
+	[[nodiscard]] bool IsFalse(bool Condition, const char* FailureMessage);
+	[[nodiscard]] bool IsFalse(bool Condition, const TCHAR* FailureMessage);
+	[[nodiscard]] bool IsFalse(bool Condition, const FString& FailureMessage);
+
+	template <typename T>
+	[[nodiscard]] bool IsNull(const T& Ptr);
+
+	template <typename T>
+	[[nodiscard]] bool IsNull(const T& Ptr, const char* FailureMessage);
+
+	template <typename T>
+	[[nodiscard]] bool IsNull(const T& Ptr, const TCHAR* FailureMessage);
+
+	template <typename T>
+	[[nodiscard]] bool IsNull(const T& Ptr, const FString& FailureMessage);
+
+	template <typename T>
+	[[nodiscard]] bool IsNotNull(const T& Ptr);
+
+	template <typename T>
+	[[nodiscard]] bool IsNotNull(const T& Ptr, const char* FailureMessage);
+
+	template <typename T>
+	[[nodiscard]] bool IsNotNull(const T& Ptr, const TCHAR* FailureMessage);
+
+	template <typename T>
+	[[nodiscard]] bool IsNotNull(const T& Ptr, const FString& FailureMessage);
+
+	template <typename TExpected, typename TActual>
+	[[nodiscard]] bool AreEqual(const TExpected& Expected, const TActual& Actual);
+
+	template <typename TExpected, typename TActual>
+	[[nodiscard]] bool AreEqual(const TExpected& Expected, const TActual& Actual, const char* FailureMessage);
+
+	template <typename TExpected, typename TActual>
+	[[nodiscard]] bool AreEqual(const TExpected& Expected, const TActual& Actual, const TCHAR* FailureMessage);
+
+	template <typename TExpected, typename TActual>
+	[[nodiscard]] bool AreEqual(const TExpected& Expected, const TActual& Actual, const FString& FailureMessage);
+
+	template <typename TExpected, typename TActual>
+	[[nodiscard]] bool AreNotEqual(const TExpected& Expected, const TActual& Actual);
+
+	template <typename TExpected, typename TActual>
+	[[nodiscard]] bool AreNotEqual(const TExpected& Expected, const TActual& Actual, const char* FailureMessage);
+
+	template <typename TExpected, typename TActual>
+	[[nodiscard]] bool AreNotEqual(const TExpected& Expected, const TActual& Actual, const TCHAR* FailureMessage);
+
+	template <typename TExpected, typename TActual>
+	[[nodiscard]] bool AreNotEqual(const TExpected& Expected, const TActual& Actual, const FString& FailureMessage);
+
+	template <typename TExpected, typename TActual, typename TEpsilon>
+	[[nodiscard]] bool IsNear(TExpected Expected, TActual Actual, TEpsilon Epsilon);
+
+	template <typename TExpected, typename TActual, typename TEpsilon>
+	[[nodiscard]] bool IsNear(TExpected Expected, TActual Actual, TEpsilon Epsilon, const char* FailureMessage);
+
+	template <typename TExpected, typename TActual, typename TEpsilon>
+	[[nodiscard]] bool IsNear(TExpected Expected, TActual Actual, TEpsilon Epsilon, const TCHAR* FailureMessage);
+
+	template <typename TExpected, typename TActual, typename TEpsilon>
+	[[nodiscard]] bool IsNear(TExpected Expected, TActual Actual, TEpsilon Epsilon, const FString& FailureMessage);
+};
+
+#include "Assert/NoDiscardAsserter.inl"
\ No newline at end of file
diff --git a/Engine/Plugins/Tests/CQTest/Source/CQTest/Public/Assert/NoDiscardAsserter.inl b/Engine/Plugins/Tests/CQTest/Source/CQTest/Public/Assert/NoDiscardAsserter.inl
new file mode 100644
--- /dev/null
+++ b/Engine/Plugins/Tests/CQTest/Source/CQTest/Public/Assert/NoDiscardAsserter.inl
@@ -0,0 +1,205 @@
+// Copyright Epic Games, Inc. All Rights Reserved.
+
+#include "Assert/NoDiscardAsserter.h"
+#include "Internationalization/Regex.h"
+#include "Assert/CQTestConvert.h"
+
+void FNoDiscardAsserter::ExpectError(FString Error, int32 Count)
+{
+	static const FRegexPattern Pattern(TEXT("([\\[\\]{}()^$.?\\\\*+|])"));
+	FRegexMatcher Matcher(Pattern, Error);
+	FString SanitizedError(Error);
+	int32 Offset = 0;
+
+	while (Matcher.FindNext())
+	{
+		int32 Pos = Matcher.GetMatchBeginning() + Offset++;
+		SanitizedError.InsertAt(Pos, TEXT("\\"));
+	}
+
+	TestRunner.AddExpectedError(SanitizedError, EAutomationExpectedErrorFlags::Contains, Count);
+}
+
+inline void FNoDiscardAsserter::ExpectErrorRegex(FString Error, int32 Count)
+{
+	TestRunner.AddExpectedError(Error, EAutomationExpectedErrorFlags::Contains, Count);
+}
+
+inline void FNoDiscardAsserter::Fail(FString Error)
+{
+	TestRunner.AddError(Error);
+}
+
+[[nodiscard]] inline bool FNoDiscardAsserter::IsTrue(bool Condition)
+{
+	return IsTrue(Condition, "condition");
+}
+
+[[nodiscard]] inline bool FNoDiscardAsserter::IsTrue(bool Condition, const char* FailureMessage)
+{
+	return IsTrue(Condition, FString(FailureMessage));
+}
+
+[[nodiscard]] inline bool FNoDiscardAsserter::IsTrue(bool Condition, const TCHAR* FailureMessage)
+{
+	return IsTrue(Condition, FString(FailureMessage));
+}
+
+[[nodiscard]] inline bool FNoDiscardAsserter::IsTrue(bool Condition, const FString& FailureMessage)
+{
+	return TestRunner.TestTrue(FailureMessage, Condition);
+}
+
+[[nodiscard]] inline bool FNoDiscardAsserter::IsFalse(bool Condition)
+{
+	return IsFalse(Condition, "condition");
+}
+
+[[nodiscard]] inline bool FNoDiscardAsserter::IsFalse(bool Condition, const char* FailureMessage)
+{
+	return IsFalse(Condition, FString(FailureMessage));
+}
+
+[[nodiscard]] inline bool FNoDiscardAsserter::IsFalse(bool Condition, const TCHAR* FailureMessage)
+{
+	return IsFalse(Condition, FString(FailureMessage));
+}
+
+[[nodiscard]] inline bool FNoDiscardAsserter::IsFalse(bool Condition, const FString& FailureMessage)
+{
+	return TestRunner.TestFalse(FString(FailureMessage), Condition);
+}
+
+template <typename T>
+[[nodiscard]] bool FNoDiscardAsserter::IsNull(const T& Ptr)
+{
+	return IsNull(Ptr, "TestNull failed");
+}
+
+template <typename T>
+[[nodiscard]] bool FNoDiscardAsserter::IsNull(const T& Ptr, const char* FailureMessage)
+{
+	return IsNull(Ptr, FString(FailureMessage));
+}
+
+template <typename T>
+[[nodiscard]] bool FNoDiscardAsserter::IsNull(const T& Ptr, const TCHAR* FailureMessage)
+{
+	return IsNull(Ptr, FString(FailureMessage));
+}
+
+template <typename T>
+[[nodiscard]] bool FNoDiscardAsserter::IsNull(const T& Ptr, const FString& FailureMessage)
+{
+	if (Ptr != nullptr)
+	{
+		Fail(FailureMessage);
+		return false;
+	}
+	return true;
+}
+
+template <typename T>
+[[nodiscard]] bool FNoDiscardAsserter::IsNotNull(const T& Ptr)
+{
+	return IsNotNull(Ptr, "TestNotNull Failed");
+}
+
+template <typename T>
+[[nodiscard]] bool FNoDiscardAsserter::IsNotNull(const T& Ptr, const char* FailureMessage)
+{
+	return IsNotNull(Ptr, FString(FailureMessage));
+}
+
+template <typename T>
+[[nodiscard]] bool FNoDiscardAsserter::IsNotNull(const T& Ptr, const TCHAR* FailureMessage)
+{
+	return IsNotNull(Ptr, FString(FailureMessage));
+}
+
+template <typename T>
+[[nodiscard]] bool FNoDiscardAsserter::IsNotNull(const T& Ptr, const FString& FailureMessage)
+{
+	if (Ptr == nullptr)
+	{
+		Fail(FailureMessage);
+		return false;
+	}
+	return true;
+}
+
+
+template <typename TExpected, typename TActual>
+[[nodiscard]] bool FNoDiscardAsserter::AreEqual(const TExpected& Expected, const TActual& Actual)
+{
+	FString Message = FString::Printf(TEXT("Expected %s to equal %s"), *CQTestConvert::ToString(Expected), *CQTestConvert::ToString(Actual));
+	return AreEqual(Expected, Actual, Message);
+}
+
+template <typename TExpected, typename TActual>
+[[nodiscard]] bool FNoDiscardAsserter::AreEqual(const TExpected& Expected, const TActual& Actual, const char* FailureMessage)
+{
+	return AreEqual(Expected, Actual, FString(FailureMessage));
+}
+
+template <typename TExpected, typename TActual>
+[[nodiscard]] bool FNoDiscardAsserter::AreEqual(const TExpected& Expected, const TActual& Actual, const TCHAR* FailureMessage)
+{
+	return AreEqual(Expected, Actual, FString(FailureMessage));
+}
+
+template <typename TExpected, typename TActual>
+[[nodiscard]] bool FNoDiscardAsserter::AreEqual(const TExpected& Expected, const TActual& Actual, const FString& FailureMessage)
+{
+	return TestRunner.TestEqual(FailureMessage, Actual, Expected); // FAutomationTestBase expects Actual then Expected
+}
+
+template <typename TExpected, typename TActual>
+[[nodiscard]] bool FNoDiscardAsserter::AreNotEqual(const TExpected& Expected, const TActual& Actual)
+{
+	FString Message = FString::Printf(TEXT("Expected %s to not equal %s"), *CQTestConvert::ToString(Expected), *CQTestConvert::ToString(Actual));
+	return AreNotEqual(Expected, Actual, Message);
+}
+
+template <typename TExpected, typename TActual>
+[[nodiscard]] bool FNoDiscardAsserter::AreNotEqual(const TExpected& Expected, const TActual& Actual, const char* FailureMessage)
+{
+	return AreNotEqual(Expected, Actual, FString(FailureMessage));
+}
+
+template <typename TExpected, typename TActual>
+[[nodiscard]] bool FNoDiscardAsserter::AreNotEqual(const TExpected& Expected, const TActual& Actual, const TCHAR* FailureMessage)
+{
+	return AreNotEqual(Expected, Actual, FString(FailureMessage));
+}
+
+template <typename TExpected, typename TActual>
+[[nodiscard]] bool FNoDiscardAsserter::AreNotEqual(const TExpected& Expected, const TActual& Actual, const FString& FailureMessage)
+{
+	return TestRunner.TestNotEqual(FailureMessage, Actual, Expected); // FAutomationTestBase expects Actual then Expected
+}
+
+template <typename TExpected, typename TActual, typename TEpsilon>
+[[nodiscard]] bool FNoDiscardAsserter::IsNear(TExpected Expected, TActual Actual, TEpsilon Epsilon)
+{
+	FString Message = FString::Printf(TEXT("Expected %s to be near %s (within %s)"), *CQTestConvert::ToString(Expected), *CQTestConvert::ToString(Actual), *CQTestConvert::ToString(Epsilon));
+	return IsNear(Expected, Actual, Epsilon, Message);
+}
+
+template <typename TExpected, typename TActual, typename TEpsilon>
+[[nodiscard]] bool FNoDiscardAsserter::IsNear(TExpected Expected, TActual Actual, TEpsilon Epsilon, const char* FailureMessage)
+{
+	return IsNear(Expected, Actual, Epsilon, FString(FailureMessage));
+}
+
+template <typename TExpected, typename TActual, typename TEpsilon>
+[[nodiscard]] inline bool FNoDiscardAsserter::IsNear(TExpected Expected, TActual Actual, TEpsilon Epsilon, const TCHAR* FailureMessage)
+{
+	return IsNear(Expected, Actual, Epsilon, FString(FailureMessage));
+}
+
+template <typename TExpected, typename TActual, typename TEpsilon>
+[[nodiscard]] inline bool FNoDiscardAsserter::IsNear(TExpected Expected, TActual Actual, TEpsilon Epsilon, const FString& FailureMessage)
+{
+	return TestRunner.TestEqual(FailureMessage, Actual, Expected, Epsilon); // FAutomationTestBase expects Actual then Expected
+}
diff --git a/Engine/Plugins/Tests/CQTest/Source/CQTest/Public/CQTest.h b/Engine/Plugins/Tests/CQTest/Source/CQTest/Public/CQTest.h
new file mode 100644
--- /dev/null
+++ b/Engine/Plugins/Tests/CQTest/Source/CQTest/Public/CQTest.h
@@ -0,0 +1,172 @@
+// Copyright Epic Games, Inc. All Rights Reserved.
+
+#pragma once
+
+#include "Misc/AutomationTest.h"
+#include "HAL/Platform.h"
+
+#include "Commands/TestCommands.h"
+#include "Commands/TestCommandBuilder.h"
+
+#include "Assert/NoDiscardAsserter.h"
+
+/*
+//Example boiler plate
+
+#include "CQTest.h"
+//TEST_CLASS(MyFixtureName, GenerateTestDirectory)
+TEST_CLASS(MyFixtureName, "Game.Example")
+{
+	//Member variables shared between tests
+
+	BEFORE_EACH()
+	{
+		//delete if empty
+	}
+
+	AFTER_EACH()
+	{
+		//delete if empty
+	}
+
+	TEST_METHOD(When_Given_Expect)
+	{
+		ASSERT_THAT(IsTrue(false));
+	}
+};
+*/
+
+namespace TestDirectoryGenerator
+{
+	CQTEST_API FString Generate(const FString& Filename);
+}
+
+static const FString GenerateTestDirectory = TEXT("");
+static constexpr uint32 DefaultFlags = EAutomationTestFlags::ApplicationContextMask | EAutomationTestFlags::ProductFilter;
+
+template <typename AsserterType>
+struct TBaseTest
+{
+	TBaseTest(FAutomationTestBase& TestRunner, bool bInitializing);
+	TBaseTest(const TBaseTest<AsserterType>& other) = delete;
+	TBaseTest& operator=(const TBaseTest<AsserterType>& rhs) = delete;
+
+	virtual ~TBaseTest() = default;
+
+	virtual void Setup() {}
+
+	virtual void TearDown() {}
+
+	// The framework expects you to pass a 'new' command, which it will own and destroy
+	void AddCommand(IAutomationLatentCommand* Cmd);
+	void AddCommand(TSharedPtr<IAutomationLatentCommand> Cmd);
+
+	virtual void RunTest(const FString& MethodName) = 0;
+	FTestCommandBuilder TestCommandBuilder;
+	FAutomationTestBase& TestRunner;
+	bool bHasLatentActions{ false };
+	bool bInitializing{ true };
+
+	AsserterType Assert;
+};
+
+template <typename AsserterType>
+struct TTestRunner;
+
+template <typename AsserterType>
+using TTestInstanceGenerator = TUniquePtr<TBaseTest<AsserterType>> (*)(TTestRunner<AsserterType>&);
+
+template <typename AsserterType>
+struct TTestRunner : public FAutomationTestBase
+{
+	TTestRunner(FString Name, int32 LineNumber, const char* FileName, FString TestDir, uint32 TestFlags, TTestInstanceGenerator<AsserterType> Factory);
+
+	uint32 GetTestFlags() const override;
+	FString GetTestSourceFileName() const override;
+	int32 GetTestSourceFileLine() const override;
+	int32 GetTestSourceFileLine(const FString& Name) const override;
+	void GetTests(TArray<FString>& OutBeautifiedNames, TArray<FString>& OutTestCommands) const override;
+
+	bool RunTest(const FString& RequestedTest) override;
+
+	int32 LineNumber = 0;
+	FString FileName;
+	FString TestDir;
+	uint32 TestFlags{ DefaultFlags };
+	bool bInitializing{ true };
+
+	TUniquePtr<TBaseTest<AsserterType>> CurrentTestPtr;
+	TArray<FString> TestNames;
+	TMap<FString, int32> TestLineNumbers{};
+	TTestInstanceGenerator<AsserterType> TestInstanceFactory;
+
+protected:
+	FString GetBeautifiedTestName() const override;
+	uint32 GetRequiredDeviceNum() const override;
+};
+
+template <typename Derived, typename AsserterType>
+struct TTest : TBaseTest<AsserterType>
+{
+	using TestMethod = void (Derived::*)();
+	using DerivedType = Derived;
+
+	TTest()
+		: TBaseTest<AsserterType>(*Derived::TestRunner, Derived::TestRunner->bInitializing) {}
+
+	void RunTest(const FString& TestName) override;
+
+	struct FFunctionRegistrar
+	{
+		FFunctionRegistrar(FString Name, TestMethod Func, int32 LineNumber)
+		{
+			DerivedType::TestRunner->TestNames.Add(Name);
+			DerivedType::Methods.Add(Name, Func);
+			DerivedType::TestRunner->TestLineNumbers.Add(Name, LineNumber);
+		}
+	};
+
+	static TUniquePtr<TBaseTest<AsserterType>> CreateTestClass(TTestRunner<AsserterType>& TestRunnerInstance)
+	{
+		DerivedType::TestRunner = &TestRunnerInstance;
+		return MakeUnique<DerivedType>();
+	}
+
+	inline static TMap<FString, TestMethod> Methods{};
+	inline static TTestRunner<AsserterType>* TestRunner{ nullptr };
+};
+
+#define _TEST_CLASS_IMPL(_ClassName, _TestDir, _BaseClass, _AsserterType, _TestFlags)                                                 \
+	struct _ClassName;                                                                                                    \
+	struct F##_ClassName##_Runner : public TTestRunner<_AsserterType>                                                     \
+	{                                                                                                                     \
+		F##_ClassName##_Runner()                                                                                          \
+			: TTestRunner(#_ClassName, __LINE__, __FILE__, _TestDir, _TestFlags, TTest<_ClassName, _AsserterType>::CreateTestClass) {} \
+	};                                                                                                                    \
+	F##_ClassName##_Runner _ClassName##_RunnerInstance;                                                                   \
+	struct _ClassName : public _BaseClass<_ClassName, _AsserterType>
+
+#define TEST_CLASS_WITH_ASSERTS(_ClassName, _TestDir, _AsserterType) _TEST_CLASS_IMPL(_ClassName, _TestDir, TTest, _AsserterType, DefaultFlags)
+
+#define TEST_CLASS(_ClassName, _TestDir) TEST_CLASS_WITH_ASSERTS(_ClassName, _TestDir, FNoDiscardAsserter)
+#define TEST_CLASS_WITH_BASE(_ClassName, _TestDir, _BaseClass) _TEST_CLASS_IMPL(_ClassName, _TestDir, _BaseClass, FNoDiscardAsserter, DefaultFlags)
+#define TEST_CLASS_WITH_FLAGS(_ClassName, _TestDir, _Flags) _TEST_CLASS_IMPL(_ClassName, _TestDir, TTest, FNoDiscardAsserter, _Flags)
+#define TEST_CLASS_WITH_BASE_AND_FLAGS(_ClassName, _TestDir, _BaseClass, _Flags) _TEST_CLASS_IMPL(_ClassName, _TestDir, _BaseClass, FNoDiscardAsserter, _Flags)
+
+#define TEST_METHOD(_MethodName)                                                            \
+	FFunctionRegistrar reg##_MethodName{ FString(#_MethodName), &DerivedType::_MethodName, __LINE__ }; \
+	void _MethodName()
+
+#define TEST(_TestName, _TestDir)        \
+	TEST_CLASS(_TestName, _TestDir)      \
+	{                                    \
+		TEST_METHOD(_TestName##_Method); \
+	};                                   \
+	void _TestName::_TestName##_Method()
+
+#define ASSERT_THAT(_assertion) if (!Assert._assertion) {return;}
+
+#define BEFORE_EACH() virtual void Setup() override
+#define AFTER_EACH() virtual void TearDown() override
+
+#include "Impl/CQTest.inl"
diff --git a/Engine/Plugins/Tests/CQTest/Source/CQTest/Public/CQTestBlueprintHelper.h b/Engine/Plugins/Tests/CQTest/Source/CQTest/Public/CQTestBlueprintHelper.h
new file mode 100644
--- /dev/null
+++ b/Engine/Plugins/Tests/CQTest/Source/CQTest/Public/CQTestBlueprintHelper.h
@@ -0,0 +1,12 @@
+// Copyright Epic Games, Inc. All Rights Reserved.
+
+#pragma once
+#include "UObject/Object.h"
+#include "UObject/Class.h"
+
+/// Utilities for working with Blueprint in C++ Tests
+struct CQTEST_API FCQTestBlueprintHelper
+{
+	UObject* FindDataBlueprint(const FString& Directory, const FString& Name);
+	UClass* GetBlueprintClass(const FString& Directory, const FString& Name);
+};
\ No newline at end of file
diff --git a/Engine/Plugins/Tests/CQTest/Source/CQTest/Public/Commands/TestCommandBuilder.h b/Engine/Plugins/Tests/CQTest/Source/CQTest/Public/Commands/TestCommandBuilder.h
new file mode 100644
--- /dev/null
+++ b/Engine/Plugins/Tests/CQTest/Source/CQTest/Public/Commands/TestCommandBuilder.h
@@ -0,0 +1,62 @@
+// Copyright Epic Games, Inc. All Rights Reserved.
+
+#pragma once
+
+#include "Misc/AutomationTest.h"
+#include "Commands/TestCommands.h"
+
+class FTestCommandBuilder
+{
+public:
+	FTestCommandBuilder() = default;
+
+	~FTestCommandBuilder()
+	{
+		check(CommandQueue.IsEmpty());
+	}
+
+	FTestCommandBuilder& Do(TFunction<void()> Action)
+	{
+		CommandQueue.Add(MakeShared<FExecute>(Action));
+		return *this;
+	}
+
+	FTestCommandBuilder& Then(TFunction<void()> Action)
+	{
+		return Do(Action);
+	}
+
+	FTestCommandBuilder& Until(FAutomationTestBase* TestRunner, TFunction<bool()> Query, FTimespan Timeout = FTimespan::FromSeconds(10))
+	{
+		CommandQueue.Add(MakeShared<FWaitUntil>(TestRunner, Query, Timeout));
+		return *this;
+	}
+
+	FTestCommandBuilder& StartWhen(FAutomationTestBase* TestRunner, TFunction<bool()> Query, FTimespan Timeout = FTimespan::FromSeconds(10))
+	{
+		return Until(TestRunner, Query, Timeout);
+	}
+
+	TSharedPtr<IAutomationLatentCommand> Build()
+	{
+		TSharedPtr<IAutomationLatentCommand> Result = nullptr;
+		if (CommandQueue.Num() == 0)
+		{
+			return Result;
+		}
+		else if (CommandQueue.Num() == 1)
+		{
+			Result = CommandQueue[0];
+		}
+		else
+		{
+			Result = MakeShared<FRunSequence>(CommandQueue);
+		}
+
+		CommandQueue.Empty();
+		return Result;
+	}
+
+private:
+	TArray<TSharedPtr<IAutomationLatentCommand>> CommandQueue{};
+};
\ No newline at end of file
diff --git a/Engine/Plugins/Tests/CQTest/Source/CQTest/Public/Commands/TestCommands.h b/Engine/Plugins/Tests/CQTest/Source/CQTest/Public/Commands/TestCommands.h
new file mode 100644
--- /dev/null
+++ b/Engine/Plugins/Tests/CQTest/Source/CQTest/Public/Commands/TestCommands.h
@@ -0,0 +1,118 @@
+// Copyright Epic Games, Inc. All Rights Reserved.
+#pragma once
+
+#include "Misc/AutomationTest.h"
+
+class FWaitUntil : public IAutomationLatentCommand
+{
+public:
+	explicit FWaitUntil(FAutomationTestBase* TestRunner, TFunction<bool()> Query, FTimespan Timeout = FTimespan::FromSeconds(10))
+		: TestBase(TestRunner)
+		, Query(MoveTemp(Query)) 
+		, Timeout(Timeout)
+		, StartTime(FDateTime::UtcNow())
+	{}
+
+	bool Update() override
+	{
+		if (Query())
+		{
+			return true;
+		}
+		else if (FDateTime::UtcNow() >= StartTime + Timeout)
+		{
+			TestBase->AddError(TEXT("Latent command timed out."));
+			return true;
+		}
+		return false;
+	}
+
+	FAutomationTestBase* TestBase;
+	TFunction<bool()> Query;
+	FTimespan Timeout;
+	FDateTime StartTime;
+};
+
+class FExecute : public IAutomationLatentCommand
+{
+public:
+	explicit FExecute(TFunction<void()> Func)
+		: Func(MoveTemp(Func)) {}
+
+	bool Update() override
+	{
+		Func();
+		return true;
+	}
+
+	TFunction<void()> Func;
+};
+
+class FRunSequence : public IAutomationLatentCommand
+{
+public:
+	FRunSequence(const TArray<TSharedPtr<IAutomationLatentCommand>>& ToAdd)
+		: Commands(ToAdd)
+	{
+	}
+
+	template <class... Cmds>
+	FRunSequence(Cmds... Commands)
+		: FRunSequence(TArray<TSharedPtr<IAutomationLatentCommand>>{ Commands... })
+	{
+	}
+
+	bool Update() override
+	{
+		// Taking notes for the next time someone needs to understand the odd behavior here
+		// From the test framework's point of view, latent commands can all be run after the test has completed
+		// So if a BeforeTest function adds latent actions which need to happen before RunTest,
+		// then RunTest, AfterTest, and TearDown must be latent actions, run in that sequence
+		// But all of the latent actions from BeforeTest must run before RunTest is executed
+		// Which means that RunTest can't be added to the latent actions queue until
+		// there are no more latent commands in the queue to run
+		// So we need to run the first command in this series, followed by asking the test framework
+		// to execute all of its latent commands.
+		// Given that the latent actions may be waiting for tickable game objects
+		// Control must be given back to the game, so we can't use a tight loop
+		// to call ExecuteLatentCommands
+		// The next challenge is that THIS command is in that list of commands
+		// So we track if this call was recursive with the bRentrantFlag
+		if (Commands.Num() == 0)
+		{
+			return true;
+		}
+		if (bRentrantFlag)
+		{
+			return true;
+		}
+
+		if (bCommandDone)
+		{
+			bRentrantFlag = true;
+			if (FAutomationTestFramework::Get().ExecuteLatentCommands())
+			{
+				FAutomationTestFramework::Get().EnqueueLatentCommand(MakeShared<FRunSequence>(Commands));
+				return true;
+			}
+			else
+			{
+				bRentrantFlag = false;
+			}
+		}
+		else
+		{
+			if (Commands[0]->Update())
+			{
+				Commands.RemoveAt(0);
+				bCommandDone = true;
+			}
+		}
+
+		return false;
+	}
+
+	TArray<TSharedPtr<IAutomationLatentCommand>> Commands;
+	bool bRentrantFlag = false;
+	bool bCommandDone = false;
+};
\ No newline at end of file
diff --git a/Engine/Plugins/Tests/CQTest/Source/CQTest/Public/Commands/TestMacros.h b/Engine/Plugins/Tests/CQTest/Source/CQTest/Public/Commands/TestMacros.h
new file mode 100644
--- /dev/null
+++ b/Engine/Plugins/Tests/CQTest/Source/CQTest/Public/Commands/TestMacros.h
@@ -0,0 +1,13 @@
+// Copyright Epic Games, Inc. All Rights Reserved.
+
+#pragma once
+
+#define TEST_EQ(Left, Right) TestRunner->TestEqual(TEXT(#Left "==" #Right), Left, Right)
+#define TEST_TRUE(Condition) TestRunner->TestTrue(TEXT(#Condition), Condition)
+#define ADD_ERROR(Error) TestRunner->AddError(Error)
+#define TEST_ERROR(Error) TestRunner->AddExpectedError(Error)
+
+#define DO(Cmd) TestCommandBuilder.Do([&]() { Cmd })
+#define START_WHEN(Query) TestCommandBuilder.StartWhen([&]() { Query })
+#define THEN(Cmd) Then([&]() { Cmd })
+#define UNTIL(Query) Until([&]() { Query })
\ No newline at end of file
diff --git a/Engine/Plugins/Tests/CQTest/Source/CQTest/Public/Impl/CQTest.inl b/Engine/Plugins/Tests/CQTest/Source/CQTest/Public/Impl/CQTest.inl
new file mode 100644
--- /dev/null
+++ b/Engine/Plugins/Tests/CQTest/Source/CQTest/Public/Impl/CQTest.inl
@@ -0,0 +1,238 @@
+// Copyright Epic Games, Inc. All Rights Reserved.
+
+namespace
+{
+	template <typename AsserterType>
+	struct TBeforeTestCommand : public IAutomationLatentCommand
+	{
+		explicit TBeforeTestCommand(TBaseTest<AsserterType>& CurrentTest)
+			: CurrentTest(CurrentTest) {}
+
+		bool Update() override
+		{
+			CurrentTest.Setup();
+			if (auto LatentActions = CurrentTest.TestCommandBuilder.Build())
+			{
+				CurrentTest.AddCommand(LatentActions);
+			}
+			return true;
+		}
+
+		TBaseTest<AsserterType>& CurrentTest;
+	};
+
+	template <typename AsserterType>
+	struct TRunTestCommand : public IAutomationLatentCommand
+	{
+		TRunTestCommand(TBaseTest<AsserterType>& CurrentTest, const FString& RequestedTest, const TTestRunner<AsserterType>& TestRunner)
+			: CurrentTest(CurrentTest)
+			, RequestedTest(RequestedTest)
+			, TestRunner(TestRunner)
+		{
+		}
+
+		bool Update() override
+		{
+			if (TestRunner.HasAnyErrors())
+			{
+				return true; // skip run if errors in BeforeTest
+			}
+
+			CurrentTest.RunTest(RequestedTest);
+			if (auto LatentActions = CurrentTest.TestCommandBuilder.Build())
+			{
+				CurrentTest.AddCommand(LatentActions);
+			}
+			return true;
+		}
+
+		TBaseTest<AsserterType>& CurrentTest;
+		const FString& RequestedTest;
+		const TTestRunner<AsserterType>& TestRunner;
+	};
+
+	template <typename AsserterType>
+	struct TAfterTestCommand : public IAutomationLatentCommand
+	{
+		explicit TAfterTestCommand(TBaseTest<AsserterType>& CurrentTest)
+			: CurrentTest(CurrentTest) {}
+
+		bool Update() override
+		{
+			CurrentTest.TearDown();
+			if (auto LatentActions = CurrentTest.TestCommandBuilder.Build())
+			{
+				CurrentTest.AddCommand(LatentActions);
+			}
+			return true;
+		}
+
+		TBaseTest<AsserterType>& CurrentTest;
+	};
+
+	template <typename AsserterType>
+	struct TTearDownRunner : public IAutomationLatentCommand
+	{
+		explicit TTearDownRunner(TTestRunner<AsserterType>& TestRunner)
+			: TestRunner(TestRunner) {}
+
+		bool Update() override
+		{
+			TestRunner.CurrentTestPtr = nullptr;
+			return true;
+		}
+
+		TTestRunner<AsserterType>& TestRunner;
+	};
+} // namespace
+
+
+template <typename AsserterType>
+inline TTestRunner<AsserterType>::TTestRunner(FString InName, int32 InLineNumber, const char* InFileName, FString InTestDir, uint32 InTestFlags, TTestInstanceGenerator<AsserterType> InFactory)
+	: FAutomationTestBase(InName, true)
+	, LineNumber(InLineNumber)
+	, FileName(InFileName)
+	, TestDir(InTestDir)
+	, TestFlags(InTestFlags)
+	, TestInstanceFactory(InFactory)
+{
+	bInitializing = true;
+	if (TestDir.Equals(GenerateTestDirectory))
+	{
+		TestDir = TestDirectoryGenerator::Generate(FileName);
+	}
+	else if (TestDir.Contains(TEXT("[GenerateTestDirectory]"), ESearchCase::IgnoreCase))
+	{
+		TestDir = TestDir.Replace(TEXT("[GenerateTestDirectory]"), *TestDirectoryGenerator::Generate(FileName), ESearchCase::IgnoreCase);
+	}
+
+	CurrentTestPtr = TestInstanceFactory(*this);
+
+	bInitializing = false;
+}
+
+template <typename AsserterType>
+inline bool TTestRunner<AsserterType>::RunTest(const FString& RequestedTest)
+{
+	if (RequestedTest.Len() == 0)
+	{
+		return false;
+	}
+
+	CurrentTestPtr = TestInstanceFactory(*this);
+	check(CurrentTestPtr != nullptr);
+	auto& CurrentTest = *CurrentTestPtr;
+
+	auto Before = MakeShared<TBeforeTestCommand<AsserterType>>(CurrentTest);
+	auto Run = MakeShared<TRunTestCommand<AsserterType>>(CurrentTest, RequestedTest, *this);
+	auto After = MakeShared<TAfterTestCommand<AsserterType>>(CurrentTest);
+	auto TearDown = MakeShared<TTearDownRunner<AsserterType>>(*this);
+
+	auto RemainingSteps = TArray<TSharedPtr<IAutomationLatentCommand>>{ Before, Run, After, TearDown };
+
+	while (RemainingSteps.Num() > 0)
+	{
+		RemainingSteps[0]->Update();
+		RemainingSteps.RemoveAt(0);
+
+		if (CurrentTestPtr != nullptr && CurrentTest.bHasLatentActions && RemainingSteps.Num() > 1)
+		{
+			// Ensure that all latent commands are run before adding the next step
+			RemainingSteps.Insert(MakeShared<FExecute>([]() { return true; }), 0);
+			CurrentTest.AddCommand(new FRunSequence(RemainingSteps));
+			return true;
+		}
+	}
+
+	return true;
+}
+
+template <typename AsserterType>
+inline FString TTestRunner<AsserterType>::GetBeautifiedTestName() const
+{
+	return FString::Printf(TEXT("%s.%s"), *TestDir, *TestName);
+}
+
+template <typename AsserterType>
+inline uint32 TTestRunner<AsserterType>::GetRequiredDeviceNum() const
+{
+	return 1;
+}
+
+template <typename AsserterType>
+inline uint32 TTestRunner<AsserterType>::GetTestFlags() const
+{
+	return TestFlags;
+}
+
+template <typename AsserterType>
+inline FString TTestRunner<AsserterType>::GetTestSourceFileName() const
+{
+	return FileName;
+}
+
+template <typename AsserterType>
+inline int32 TTestRunner<AsserterType>::GetTestSourceFileLine() const
+{
+	return LineNumber;
+}
+
+template <typename AsserterType>
+inline int32 TTestRunner<AsserterType>::GetTestSourceFileLine(const FString& Name) const
+{
+	FString TestParam(Name);
+	int32 Pos = Name.Find(TEXT(" "));
+	if (Pos != INDEX_NONE)
+	{
+		TestParam = Name.RightChop(Pos + 1);
+	}
+	if (TestNames.Contains(TestParam))
+	{
+		return TestLineNumbers[TestParam];
+	}
+	return GetTestSourceFileLine();
+}
+
+template <typename AsserterType>
+inline void TTestRunner<AsserterType>::GetTests(TArray<FString>& OutBeautifiedNames, TArray<FString>& OutTestCommands) const
+{
+	for (const auto& testName : TestNames)
+	{
+		OutBeautifiedNames.Add(testName);
+		OutTestCommands.Add(testName);
+	}
+}
+
+/////////////
+
+template <typename Derived, typename AsserterType>
+inline void TTest<Derived, AsserterType>::RunTest(const FString& TestName)
+{
+	auto TestMethod = Methods[TestName];
+	Derived& Self = static_cast<Derived&>(*this);
+
+	(Self.*(TestMethod))();
+}
+
+template <typename AsserterType>
+inline TBaseTest<AsserterType>::TBaseTest(FAutomationTestBase& TestRunner, bool bInitializing)
+	: TestCommandBuilder(FTestCommandBuilder{})
+	, TestRunner(TestRunner)
+	, bInitializing(bInitializing)
+	, Assert(AsserterType{ TestRunner })
+{
+}
+
+template <typename AsserterType>
+inline void TBaseTest<AsserterType>::AddCommand(IAutomationLatentCommand* Cmd)
+{
+	TestRunner.AddCommand(Cmd);
+	bHasLatentActions = true;
+}
+
+template <typename AsserterType>
+inline void TBaseTest<AsserterType>::AddCommand(TSharedPtr<IAutomationLatentCommand> Cmd)
+{
+	FAutomationTestFramework::Get().EnqueueLatentCommand(Cmd);
+	bHasLatentActions = true;
+}
diff --git a/Engine/Plugins/Tests/CQTest/Source/CQTest/Public/MapTestSpawner.h b/Engine/Plugins/Tests/CQTest/Source/CQTest/Public/MapTestSpawner.h
new file mode 100644
--- /dev/null
+++ b/Engine/Plugins/Tests/CQTest/Source/CQTest/Public/MapTestSpawner.h
@@ -0,0 +1,26 @@
+// Copyright Epic Games, Inc. All Rights Reserved.
+
+#pragma once
+
+#include "SpawnHelper.h"
+
+/// Class for spawning Actors in a named Map/Level
+struct CQTEST_API FMapTestSpawner : public FSpawnHelper
+{
+	FMapTestSpawner(const FString& MapDirectory, const FString& MapName)
+		: MapDirectory(MapDirectory)
+		, MapName(MapName)
+	{
+	}
+
+	void AddWaitUntilLoadedCommand(FAutomationTestBase* TestRunner);
+	APawn* FindFirstPlayerPawn();
+
+protected:
+    virtual UWorld* CreateWorld() override;
+
+private:
+	FString MapDirectory;
+	FString MapName;
+	UWorld* PieWorld{ nullptr };
+};
\ No newline at end of file
diff --git a/Engine/Plugins/Tests/CQTest/Source/CQTest/Public/ObjectBuilder.h b/Engine/Plugins/Tests/CQTest/Source/CQTest/Public/ObjectBuilder.h
new file mode 100644
--- /dev/null
+++ b/Engine/Plugins/Tests/CQTest/Source/CQTest/Public/ObjectBuilder.h
@@ -0,0 +1,320 @@
+// Copyright Epic Games, Inc. All Rights Reserved.
+#pragma once
+#include "Misc/Build.h"
+
+#if WITH_AUTOMATION_WORKER
+
+#include "Containers/Array.h"
+#include "UObject/UObjectBaseUtility.h"
+#include "UObject/UObjectGlobals.h"
+#include "UObject/Package.h"
+
+#include "SpawnHelper.h"
+
+DEFINE_LOG_CATEGORY_STATIC(LogObjectBuilder, Log, Log);
+
+template <typename TUObject>
+class TObjectBuilder final
+{
+public:
+	TObjectBuilder(UObject* InOwner = static_cast<UObject*>(GetTransientPackage()))
+		: Object(NewObject<TUObject>(InOwner))
+	{
+		static_assert(std::is_base_of<UObject, TUObject>::value, "Template param must be a UObject type");
+		static_assert(!std::is_base_of<AActor, TUObject>::value, "Template param must not be an Actor type, supply a FSpawnHelper to the constructor for Actor configuration");
+		check(Object);
+	}
+
+	TObjectBuilder(FSpawnHelper& Spawner, UClass* Clazz = nullptr)
+		: Object(nullptr)
+	{
+		static_assert(std::is_base_of<AActor, TUObject>::value, "Template param must be an Actor type, use the default constructor for non-actor UObjects");
+
+		FActorSpawnParameters DeferredConstructionParams;
+		DeferredConstructionParams.bDeferConstruction = true;
+
+		Object = &Spawner.SpawnActor<TUObject>(DeferredConstructionParams, Clazz);
+	}
+
+	template <typename T>
+	TObjectBuilder<TUObject>& SetParam(const FName& InPropertyName, T InValue)
+	{
+		using Raw = std::decay_t<T>;
+
+		if (!Object)
+		{
+			UE_LOG(LogObjectBuilder, Error, TEXT("Tried to SetParam on property %s after Spawning. ObjectBuilder can only be used to paramaterise pre-spawn"), *InPropertyName.ToString());
+			return *this;
+		}
+
+		auto TrySet = [&](auto Property) {
+			if (!Property)
+			{
+				UE_LOG(LogObjectBuilder, Error, TEXT("Failed to find %s property %s on object %s, check type and Property name are correct"), *GetTypeName<Raw>(), *InPropertyName.ToString(), *Object->GetName());
+				return;
+			}
+			if (!IsCompatibleWith<Raw>(*Property))
+			{
+				UE_LOG(LogObjectBuilder, Error, TEXT("Type mismatch: Tried to set %s property %s on object %s to a %s"), *Property->GetCPPType(nullptr, 0), *InPropertyName.ToString(), *Object->GetName(), *GetTypeName<Raw>());
+				return;
+			}
+
+			if (auto TypeMatched = Property->template ContainerPtrToValuePtr<Raw>(Object))
+			{
+				*TypeMatched = InValue;
+			}
+			else
+			{
+				UE_LOG(LogObjectBuilder, Error, TEXT("Container type mismatch: Could not cast from %s to %s"), *Property->GetCPPType(nullptr, 0), *GetTypeName<Raw>());
+			}
+		};
+
+		if constexpr (TIsArray<Raw>::Value)
+		{
+			TrySet(FindFProperty<FArrayProperty>(Object->GetClass(), InPropertyName));
+		}
+		else if constexpr (TIsTMap<Raw>::Value)
+		{
+			TrySet(FindFProperty<FMapProperty>(Object->GetClass(), InPropertyName));
+		}
+		else if constexpr (TIsTSet<Raw>::Value)
+		{
+			TrySet(FindFProperty<FSetProperty>(Object->GetClass(), InPropertyName));
+		}
+		else
+		{
+			if (auto Property = FindFProperty<FProperty>(Object->GetClass(), InPropertyName))
+			{
+				if (this->template IsCompatibleWith<Raw>(*Property))
+				{
+					Property->SetValue_InContainer(Object, &InValue);
+				}
+				else
+				{
+					UE_LOG(LogObjectBuilder, Error, TEXT("Type mismatch: Tried to set %s property %s on object %s to a %s"), *Property->GetCPPType(nullptr, 0), *InPropertyName.ToString(), *Object->GetName(), *this->template GetTypeName<Raw>());
+				}
+			}
+			else
+			{
+				UE_LOG(LogObjectBuilder, Error, TEXT("Failed to find %s property %s on object %s, check type and Property name are correct"), *GetTypeName<Raw>(), *InPropertyName.ToString(), *Object->GetName());
+			}
+		}
+
+		return *this;
+	}
+
+	template<typename T = TUObject>
+	T& Spawn(FTransform InTransform = FTransform::Identity) 
+	{
+		if (!Object)
+		{
+			UE_LOG(LogObjectBuilder, Error, TEXT("Tried to spawn from builder multiple times.  Builder can only spawn a single object"));
+			auto* DefaultObject = NewObject<T>();
+			return *DefaultObject;
+		}
+
+		if constexpr (std::is_base_of_v<AActor, T>)
+		{
+			Object->FinishSpawning(InTransform);
+		}
+
+		auto& ObjectRef = *Object;
+		Object = nullptr;
+		return ObjectRef;
+	}
+
+	template<typename TComponent, typename T = TUObject>
+	TObjectBuilder<TUObject>& AddComponentTo(TComponent* InComponentToAdd = nullptr) 
+	{
+		static_assert(std::is_base_of_v<AActor, T>, "Can only add components to AActors");
+		static_assert(std::is_base_of<UActorComponent, TComponent>::value, "Template param must be a UActorComponent type");
+		if (!Object)
+		{
+			UE_LOG(LogObjectBuilder, Error, TEXT("Tried to AddComponentTo actor after Spawning. ObjectBuilder can only be used to paramaterise pre-spawn"));
+			return *this;
+		}
+
+		TComponent* Component = InComponentToAdd;
+		if (Component)
+		{
+			Component->Rename(nullptr, Object);
+		}
+		else
+		{
+			Component = NewObject<TComponent>(Object);
+			check(Component);
+		}
+
+		Component->RegisterComponent();
+		return *this;
+	}
+
+	template<typename TChildActorType, typename T = TUObject>
+	TObjectBuilder<TUObject>& AddChildActorComponentTo()
+	{
+		static_assert(std::is_base_of_v<AActor, T>, "Can only add components to AActors");
+		static_assert(std::is_base_of<AActor, TChildActorType>::value, "Template param must be an Actor type");
+		if (!Object)
+		{
+			UE_LOG(LogObjectBuilder, Error, TEXT("Tried to AddChildActorComponentTo actor after Spawning. ObjectBuilder can only be used to paramaterise pre-spawn"));
+			return *this;
+		}
+
+		auto* Component = NewObject<UChildActorComponent>(Object);
+		Component->SetChildActorClass(TChildActorType::StaticClass());
+		Component->RegisterComponent();
+		return *this;
+	}
+
+private:
+	TUObject* Object;
+
+	template <typename T>
+	bool IsCompatibleWith(const FProperty& Prop) const
+	{
+		using Raw = std::decay_t<T>;
+		if constexpr (std::is_pointer_v<T> && std::is_base_of_v<UObject, std::remove_pointer_t<Raw>>)
+		{
+			if (auto AsProperty = CastField<FObjectProperty>(&Prop))
+			{
+				return std::remove_pointer_t<Raw>::StaticClass()->IsChildOf(AsProperty->PropertyClass);
+			}
+		}
+		else if constexpr (TIsTObjectPtr<Raw>::Value)
+		{
+			if (auto AsProperty = CastField<FObjectPtrProperty>(&Prop))
+			{
+				return Raw::ElementType::StaticClass()->IsChildOf(AsProperty->PropertyClass);
+			}
+		}
+		else if constexpr (TIsTArray<Raw>::Value)
+		{
+			if (auto AsProperty = CastField<FArrayProperty>(&Prop))
+			{
+				return IsCompatibleWith<typename Raw::ElementType>(*AsProperty->Inner);
+			}
+		}
+		else if constexpr (TIsTSet<Raw>::Value)
+		{
+			if (auto AsProperty = CastField<FSetProperty>(&Prop))
+			{
+				return IsCompatibleWith<typename Raw::ElementType>(*AsProperty->ElementProp);
+			}
+		}
+		else if constexpr (TIsTMap<Raw>::Value)
+		{
+			if (auto AsProperty = CastField<FMapProperty>(&Prop))
+			{
+				using Key = typename Raw::KeyInitType;
+				using Value = typename Raw::ValueInitType;
+				return IsCompatibleWith<Key>(*AsProperty->KeyProp) && IsCompatibleWith<Value>(*AsProperty->ValueProp);
+			}
+		}
+		else if constexpr (TIsEnum<Raw>::Value)
+		{
+			if (auto AsProperty = CastField<FEnumProperty>(&Prop))
+			{
+				return AsProperty->GetUnderlyingProperty()->GetSize() == sizeof(T);
+			}
+		}
+		else if constexpr (std::is_base_of_v<UObject, Raw>)
+		{
+			if (auto AsProperty = CastField<FObjectProperty>(&Prop))
+			{
+				return Raw::StaticClass()->IsChildOf(AsProperty->PropertyClass);
+			}
+		}
+		else if constexpr (std::is_same_v<FVector, Raw>)
+		{
+			if (auto AsProperty = CastField<FStructProperty>(&Prop))
+			{
+				return AsProperty->Struct->GetFName() == NAME_Vector;
+			}
+		}
+		else if constexpr (std::is_same_v<FStructProperty, typename TPropType<Raw>::Value>)
+		{
+			if (auto AsProperty = CastField<FStructProperty>(&Prop))
+			{
+				return AsProperty->Struct->GetName() == Raw::StaticStruct()->GetName();
+			}
+		}
+		else
+		{
+			using PropType = typename TPropType<Raw>::Value;
+
+			if (auto AsProperty = CastField<PropType>(&Prop))
+			{
+				return PropType::StaticClass()->IsChildOf(AsProperty->GetClass());
+			}
+		}
+
+		return false;
+	}
+
+	template<typename T>
+	FString GetTypeName() const 
+	{
+		using Raw = std::decay_t<T>;
+		if constexpr (std::is_pointer_v<Raw> && std::is_base_of_v<UObject, std::remove_pointer_t<Raw>>)
+		{
+			return FString::Printf(TEXT("%s*"), *std::remove_pointer_t<Raw>::StaticClass()->GetName());
+		}
+		else if constexpr (TIsTObjectPtr<Raw>::Value)
+		{
+			using PtrType = typename Raw::ElementType;
+			return FString::Printf(TEXT("TObjectPtr<%s>"), *GetTypeName<PtrType>());
+		}
+		else if constexpr (TIsTArray<Raw>::Value)
+		{
+			using ElementType = typename Raw::ElementType;
+			return FString::Printf(TEXT("TArray<%s>"), *GetTypeName<ElementType>());
+		}
+		else if constexpr (TIsTSet<Raw>::Value)
+		{
+			using ElementType = typename Raw::ElementType;
+			return FString::Printf(TEXT("TSet<%s>"), *GetTypeName<ElementType>());
+		}
+		else if constexpr (TIsTMap<Raw>::Value)
+		{
+			using Key = typename Raw::KeyInitType;
+			using Value = typename Raw::ValueInitType;
+			return FString::Printf(TEXT("TMap<%s, %s>"), *GetTypeName<Key>(), *GetTypeName<Value>());
+		}
+		else if constexpr (TIsEnum<Raw>::Value)
+		{
+			return TEXT("EnumProperty");
+		}
+		else if constexpr (std::is_same_v<FVector, Raw>)
+		{
+			return TEXT("FVector");
+		}
+		else if constexpr (std::is_base_of_v<UObject, Raw>)
+		{
+			return FString::Printf(TEXT("%s"), *Raw::StaticClass()->GetName());
+		}
+		else if constexpr (std::is_same_v<FStructProperty, typename TPropType<Raw>::Value>)
+		{
+			return Raw::StaticStruct()->GetName();
+		}
+		else
+		{
+			return TNameOf<Raw>::GetName();
+		}
+	}
+
+	template<typename T> struct TPropType { using Value = FStructProperty; };
+	template <> struct TPropType<int8> { using Value = FInt8Property; };
+	template <>	struct TPropType<uint8>	{ using Value = FByteProperty; };
+	template <>	struct TPropType<int16>	{ using Value = FInt16Property;	};
+	template <>	struct TPropType<uint16> { using Value = FUInt16Property; };
+	template <>	struct TPropType<int32>	{ using Value = FIntProperty; };
+	template <>	struct TPropType<uint32> { using Value = FUInt32Property; };
+	template <>	struct TPropType<int64>	{ using Value = FInt64Property;	};
+	template <>	struct TPropType<uint64> { using Value = FUInt64Property; };
+	template <>	struct TPropType<bool> { using Value = FBoolProperty; };
+	template <>	struct TPropType<float>	{ using Value = FFloatProperty;	};
+	template <>	struct TPropType<double> { using Value = FDoubleProperty; };
+	template <>	struct TPropType<FName>	{ using Value = FNameProperty; };
+};
+
+#endif // WITH_AUTOMATION_WORKER
diff --git a/Engine/Plugins/Tests/CQTest/Source/CQTest/Public/SpawnHelper.h b/Engine/Plugins/Tests/CQTest/Source/CQTest/Public/SpawnHelper.h
new file mode 100644
--- /dev/null
+++ b/Engine/Plugins/Tests/CQTest/Source/CQTest/Public/SpawnHelper.h
@@ -0,0 +1,57 @@
+// Copyright Epic Games, Inc. All Rights Reserved.
+
+#pragma once
+
+#include "Engine/World.h"
+#include "GameFramework/Actor.h"
+#include <UObject/GCObject.h>
+
+struct CQTEST_API FSpawnHelper
+{
+	virtual ~FSpawnHelper();
+
+	template <typename ActorType>
+	ActorType& SpawnActor(const FActorSpawnParameters& SpawnParameters = FActorSpawnParameters(), UClass* Class = nullptr)
+	{
+		return SpawnActorAtInWorld<ActorType>(GetWorld(), FVector::ZeroVector, FRotator::ZeroRotator, SpawnParameters, Class);
+	}
+
+	template <typename ActorType>
+	ActorType& SpawnActorAt(FVector const& Location, FRotator const& Rotation, const FActorSpawnParameters& SpawnParameters = FActorSpawnParameters(), UClass* Class = nullptr)
+	{
+		return SpawnActorAtInWorld<ActorType>(GetWorld(), Location, Rotation, SpawnParameters, Class);
+	}
+
+	template <typename ObjectType>
+	ObjectType& SpawnObject()
+	{
+		static_assert(!TIsDerivedFrom<ObjectType, AActor>::IsDerived, "Use SpawnActor to spawn AActors.");
+		static_assert(TIsDerivedFrom<ObjectType, UObject>::IsDerived, "Objects must derive from UObject.");
+		ObjectType* const Object = NewObject<ObjectType>();
+		check(Object != nullptr);
+		SpawnedObjects.Add(Object);
+		return *Object;
+	}
+
+protected:
+	virtual UWorld* CreateWorld() = 0;
+
+	TArray<TWeakObjectPtr<AActor>> SpawnedActors{};
+	TArray<TWeakObjectPtr<UObject>> SpawnedObjects{};
+
+	UWorld& GetWorld();
+	UWorld* GameWorld{ nullptr };
+
+private:
+	template <typename ActorType>
+	ActorType& SpawnActorAtInWorld(UWorld& World, const FVector& Location, const FRotator& Rotation, const FActorSpawnParameters& SpawnParameters = FActorSpawnParameters(), UClass* Class = nullptr)
+	{
+		static_assert(TIsDerivedFrom<ActorType, AActor>::IsDerived, "Provided type does not derive from AActor");
+
+		ActorType* const Actor = Class ? World.SpawnActor<ActorType>(Class, Location, Rotation, SpawnParameters) : World.SpawnActor<ActorType>(Location, Rotation, SpawnParameters);
+
+		check(Actor != nullptr);
+		SpawnedActors.Add(Actor);
+		return *Actor;
+	}
+};
\ No newline at end of file
diff --git a/Engine/Plugins/Tests/CQTest/Source/CQTest/Public/TestGameInstance.h b/Engine/Plugins/Tests/CQTest/Source/CQTest/Public/TestGameInstance.h
new file mode 100644
--- /dev/null
+++ b/Engine/Plugins/Tests/CQTest/Source/CQTest/Public/TestGameInstance.h
@@ -0,0 +1,16 @@
+// Copyright Epic Games, Inc. All Rights Reserved.
+
+#pragma once
+
+#include "Engine/GameInstance.h"
+
+#include "TestGameInstance.generated.h"
+
+UCLASS(config=Game, transient, BlueprintType, Blueprintable)
+class CQTEST_API UTestGameInstance : public UGameInstance
+{
+	GENERATED_UCLASS_BODY()
+
+public:
+	void InitializeForTest(UWorld& InWorld);
+};
\ No newline at end of file
diff --git a/Engine/Plugins/Tests/CQTest/Source/CQTestTests/CQTestTests.Build.cs b/Engine/Plugins/Tests/CQTest/Source/CQTestTests/CQTestTests.Build.cs
new file mode 100644
--- /dev/null
+++ b/Engine/Plugins/Tests/CQTest/Source/CQTestTests/CQTestTests.Build.cs
@@ -0,0 +1,19 @@
+// Copyright Epic Games, Inc. All Rights Reserved.
+
+using UnrealBuildTool;
+
+public class CQTestTests : ModuleRules
+{
+	public CQTestTests(ReadOnlyTargetRules Target)
+		: base(Target)
+	{
+        PrivateDependencyModuleNames.AddRange(
+			new string[] {
+					"Core",
+					"CoreUObject",
+					"Engine",
+					"CQTest"
+				 }
+			);
+	}
+}
diff --git a/Engine/Plugins/Tests/CQTest/Source/CQTestTests/CQTestTestsModule.cpp b/Engine/Plugins/Tests/CQTest/Source/CQTestTests/CQTestTestsModule.cpp
new file mode 100644
--- /dev/null
+++ b/Engine/Plugins/Tests/CQTest/Source/CQTestTests/CQTestTestsModule.cpp
@@ -0,0 +1,5 @@
+// Copyright Epic Games, Inc. All Rights Reserved.
+
+#include "Modules/ModuleManager.h"
+
+IMPLEMENT_MODULE( FDefaultModuleImpl, CQTestTests )
diff --git a/Engine/Plugins/Tests/CQTest/Source/CQTestTests/Private/AssertionTests.cpp b/Engine/Plugins/Tests/CQTest/Source/CQTestTests/Private/AssertionTests.cpp
new file mode 100644
--- /dev/null
+++ b/Engine/Plugins/Tests/CQTest/Source/CQTestTests/Private/AssertionTests.cpp
@@ -0,0 +1,293 @@
+// Copyright Epic Games, Inc. All Rights Reserved.
+
+#include "CQTest.h"
+#include "Assert/NoDiscardAsserter.h"
+
+DECLARE_LOG_CATEGORY_CLASS(AssertionTests, Log, All);
+
+namespace CQTests
+{
+	TEST_CLASS(NoDiscardAssert, "TestFramework.CQTest")
+	{
+		const char* ExpectedError = "Expected";
+
+		TEST_METHOD(AssertFail_WithMessage_AddsError)
+		{
+			Assert.ExpectError(TEXT(""));
+			Assert.Fail(TEXT(""));
+		}
+
+		TEST_METHOD(Assertions_Accept_RawStrings)
+		{
+			Assert.ExpectError("Hello World");
+			Assert.Fail("Hello World");
+		}
+
+		TEST_METHOD(Assertions_Accept_TCharArrays)
+		{
+			Assert.ExpectError(TEXT("Hello World"));
+			Assert.Fail(TEXT("Hello World"));
+		}
+
+		TEST_METHOD(Assertions_Accept_FStrings)
+		{
+			FString message = TEXT("Hello World");
+			Assert.ExpectError(message);
+			Assert.Fail(TEXT("Hello WorldA"));
+		}
+
+		TEST_METHOD(AssertExpectError_WithMultipleErrors_Succeeds)
+		{
+			Assert.ExpectError("", 3);
+			Assert.Fail("One");
+			Assert.Fail("Two");
+			Assert.Fail("Three");
+		}
+
+		TEST_METHOD(AsserterExpectError_WithZeroExpected_AcceptsAnyNumber)
+		{
+			Assert.ExpectError("", 0);
+			Assert.Fail("One");
+			Assert.Fail("Two");
+			Assert.Fail("Three");
+		}
+
+		TEST_METHOD(AssertExpectError_WithMatchingError_Succeeds)
+		{
+			Assert.ExpectError("Hello World");
+			Assert.Fail("Hello World");
+		}
+
+		TEST_METHOD(AssertExpectError_WithRegexSymbols_EscapesRegex)
+		{
+			Assert.ExpectError("[^abc]");
+			Assert.Fail("[^abc]");
+		}
+
+		TEST_METHOD(AssertExpectErrorRegex_WithRegex_Succeeds)
+		{
+			Assert.ExpectErrorRegex("\\w+");
+			Assert.Fail("abc");
+		}
+
+		TEST_METHOD(AssertTrue_WithTrue_Succeeds)
+		{
+			ASSERT_THAT(IsTrue(true));
+		}
+		TEST_METHOD(AssertTrue_WithTrueAndErrorMessage_DoesNotAddErrorMessage)
+		{
+			ASSERT_THAT(IsTrue(true, "Unexpected"));
+		}
+		TEST_METHOD(AssertTrue_WithFalse_AddsError)
+		{
+			Assert.ExpectError(TEXT(""));
+			ASSERT_THAT(IsTrue(false));
+		}
+		TEST_METHOD(AssertTrue_WithFalseAndError_AddSpecificError)
+		{
+			Assert.ExpectError(FString(ExpectedError));
+			ASSERT_THAT(IsTrue(false, ExpectedError));
+		}
+
+		TEST_METHOD(AssertEqual_WithSameInts_Succeeds)
+		{
+			ASSERT_THAT(AreEqual(42, 42));
+		}
+		TEST_METHOD(AssertEqual_WithSameIntsAndErrorMessage_DoesNotAddErrorMessage)
+		{
+			ASSERT_THAT(AreEqual(42, 42, "Unexpected"));
+		}
+		TEST_METHOD(AssertEqual_WithDifferentInts_AddsError)
+		{
+			Assert.ExpectError(TEXT(""));
+			ASSERT_THAT(AreEqual(42, 0));
+		}
+		TEST_METHOD(AssertEqual_WithDifferentIntsAndErrorMessage_AddsSpecificError)
+		{
+			Assert.ExpectError(FString(ExpectedError));
+			ASSERT_THAT(AreEqual(42, 0, ExpectedError));
+		}
+
+		TEST_METHOD(AssertEqual_WithSameFStrings_Succeeds)
+		{
+			ASSERT_THAT(AreEqual(TEXT("Hello"), TEXT("Hello")));
+		}
+		TEST_METHOD(AssertEqual_WithDifferentFStrings_AddsError)
+		{
+			Assert.ExpectError(TEXT(""));
+			ASSERT_THAT(AreEqual(TEXT("Hello"), TEXT("World")));
+		}
+		TEST_METHOD(AssertEqual_WithDifferentCases_AddsError)
+		{
+			Assert.ExpectError(TEXT(""));
+			ASSERT_THAT(AreEqual(TEXT("hello"), TEXT("HELLO")));
+		}
+		TEST_METHOD(AssertEqual_WithDifferentFStringsAndErrorMessage_AddsSpecificError)
+		{
+			Assert.ExpectError(FString(ExpectedError));
+			ASSERT_THAT(AreEqual(TEXT("Hello"), TEXT("World"), ExpectedError));
+		}
+		TEST_METHOD(AssertNear_WithSameNumbers_Succeeds)
+		{
+			ASSERT_THAT(IsNear(3.14, 3.14, 0.001));
+		}
+		TEST_METHOD(AssertNear_WithSameNumbersAndErrorMessage_DoesNotAddErrorMessage)
+		{
+			ASSERT_THAT(IsNear(3.14, 3.14, 0.001, "Unexpected"));
+		}
+		TEST_METHOD(AssertNear_WithSimilarNumbers_Succeeds)
+		{
+			ASSERT_THAT(IsNear(3.0, 3.1, 1.0));
+		}
+		TEST_METHOD(AssertNear_WithDifferentNumbers_AddsError)
+		{
+			Assert.ExpectError(TEXT(""));
+			ASSERT_THAT(IsNear(1.0, 2.0, 0.001));
+		}
+		TEST_METHOD(AssertNear_WithDifferentNumbersAndError_AddsSpecificError)
+		{
+			Assert.ExpectError(FString(ExpectedError));
+			ASSERT_THAT(IsNear(1.0, 2.0, 0.001, ExpectedError));
+		}
+
+		TEST_METHOD(AssertNull_WithNullPtr_Succeeds)
+		{
+			int* ptr = nullptr;
+			ASSERT_THAT(IsNull(ptr));
+		}
+		TEST_METHOD(AssertNull_WithNullPtrAndErrorMessage_DoesNotAddError)
+		{
+			int* ptr = nullptr;
+			ASSERT_THAT(IsNull(ptr, "Unexpected"));
+		}
+		TEST_METHOD(AssertNull_WithNonNull_AddsError)
+		{
+			int val = 42;
+			int* ptr = &val;
+			Assert.ExpectError(TEXT(""));
+			ASSERT_THAT(IsNull(ptr));
+		}
+		TEST_METHOD(AssertNull_WithNonNullAndErrorMessage_AddsSpecificErrorMessage)
+		{
+			int val = 42;
+			int* ptr = &val;
+			Assert.ExpectError(FString(ExpectedError));
+			ASSERT_THAT(IsNull(ptr, ExpectedError));
+		}
+		TEST_METHOD(AssertNull_WithInvalidSharedPtr_Succeeds)
+		{
+			auto ptr = TSharedPtr<int>{};
+			ASSERT_THAT(IsNull(ptr));
+		}
+		TEST_METHOD(AssertNull_WithValidSharedPtr_AddsError)
+		{
+			TSharedPtr<int> ptr = MakeShared<int>(42);
+			Assert.ExpectError(TEXT(""));
+			ASSERT_THAT(IsNull(ptr));
+		}
+		TEST_METHOD(AssertNull_WithInvalidUniquePtr_Succeeds)
+		{
+			TUniquePtr<int> ptr;
+			ASSERT_THAT(IsNull(ptr));
+		}
+		TEST_METHOD(AssertNull_WithValidUniquePtr_AddsError)
+		{
+			TUniquePtr<int> ptr = MakeUnique<int>(42);
+			Assert.ExpectError(TEXT(""));
+			ASSERT_THAT(IsNull(ptr));
+		}
+
+		TEST_METHOD(AssertNotNull_WithNullPtr_AddsError)
+		{
+			int* ptr = nullptr;
+			Assert.ExpectError(TEXT(""));
+			ASSERT_THAT(IsNotNull(ptr));
+		}
+		TEST_METHOD(AssertNotNull_WithNullPtrAndError_AddsSpecificError)
+		{
+			int* ptr = nullptr;
+			Assert.ExpectError(FString(ExpectedError));
+			ASSERT_THAT(IsNotNull(ptr, ExpectedError));
+		}
+		TEST_METHOD(AssertNotNull_WithNonNull_Succeeds)
+		{
+			int val = 42;
+			int* ptr = &val;
+			ASSERT_THAT(IsNotNull(ptr));
+		}
+		TEST_METHOD(AssertNotNull_WithNonNullAndErrorMessage_DoesNotAddErrorMessage)
+		{
+			int val = 42;
+			int* ptr = &val;
+			ASSERT_THAT(IsNotNull(ptr, "Unexpected"));
+		}
+		TEST_METHOD(AssertNotNull_WithInvalidSharedPtr_AddsError)
+		{
+			auto ptr = TSharedPtr<int>{};
+			Assert.ExpectError(TEXT(""));
+			ASSERT_THAT(IsNotNull(ptr));
+		}
+		TEST_METHOD(AssertNotNull_WithValidSharedPtr_Succeeds)
+		{
+			TSharedPtr<int> ptr = MakeShared<int>(42);
+			ASSERT_THAT(IsNotNull(ptr));
+		}
+		TEST_METHOD(AssertNotNull_WithInvalidUniquePtr_AddsError)
+		{
+			TUniquePtr<int> ptr;
+			Assert.ExpectError(TEXT(""));
+			ASSERT_THAT(IsNotNull(ptr));
+		}
+		TEST_METHOD(AssertNotNull_WithValidUniquePtr_Succeeds)
+		{
+			TUniquePtr<int> ptr = MakeUnique<int>(42);
+			ASSERT_THAT(IsNotNull(ptr));
+		}
+	};
+	
+} // namespace CQTests
+
+struct FCustomType
+{
+	FCustomType(const FString& InName)
+		: Name(InName)
+	{
+	}
+
+	FString Name;
+
+	bool operator==(const FCustomType& other) const
+	{
+		return Name == other.Name;
+	}
+	bool operator!=(const FCustomType& other) const
+	{
+		return !(*this == other);
+	}
+};
+
+template <>
+FString CQTestConvert::ToString(const FCustomType& obj)
+{
+	return obj.Name;
+}
+
+TEST_CLASS(AssertMessage, "TestFramework.CQTest")
+{
+	TEST_METHOD(AssertAreEqual_WithUnequalCustomTypes_PrintObjectNames)
+	{
+		FCustomType left = { "left" };
+		FCustomType right = { "right" };
+
+		Assert.ExpectError("Expected left to equal right");
+		ASSERT_THAT(AreEqual(left, right));
+	}
+
+	TEST_METHOD(AssertAreNotEqual_WithEqualCustomTypes_PrintObjectNames)
+	{
+		FCustomType obj = { "Object" };
+
+		Assert.ExpectError("Expected Object to not equal Object");
+		ASSERT_THAT(AreNotEqual(obj, obj));
+	}
+};
diff --git a/Engine/Plugins/Tests/CQTest/Source/CQTestTests/Private/CQTestConvertTests.cpp b/Engine/Plugins/Tests/CQTest/Source/CQTestTests/Private/CQTestConvertTests.cpp
new file mode 100644
--- /dev/null
+++ b/Engine/Plugins/Tests/CQTest/Source/CQTestTests/Private/CQTestConvertTests.cpp
@@ -0,0 +1,202 @@
+// Copyright Epic Games, Inc. All Rights Reserved.
+
+#include "CQTest.h"
+
+static const FString OverrideStr = FString("Override");
+static const FString ToStringStr = FString("ToString");
+constexpr auto OStreamStr = "OStream";
+
+struct TypeWithExplicitOverride
+{
+};
+
+template <>
+FString CQTestConvert::ToString(const TypeWithExplicitOverride&)
+{
+	return OverrideStr;
+}
+
+struct TypeWithToString
+{
+	FString ToString() const
+	{
+		return ToStringStr;
+	}
+};
+
+struct TypeWithOStream
+{
+};
+
+CQTestConvert::PrintStream& operator<<(CQTestConvert::PrintStream& stream, const TypeWithOStream& obj) {
+	stream << OStreamStr;
+	return stream;
+}
+
+struct TypeWithToStringAndOverride
+{
+	FString ToString() const {
+		return ToStringStr;
+	}
+};
+
+template<>
+FString CQTestConvert::ToString(const TypeWithToStringAndOverride&) {
+	return OverrideStr;
+}
+
+struct TypeWithToStringAndOStream
+{
+	FString ToString() const {
+		return ToStringStr;
+	}
+	CQTestConvert::PrintStream& operator<<(CQTestConvert::PrintStream& stream)
+	{
+		stream << OStreamStr;
+		return stream;
+	}
+};
+
+enum UnscopedEnumNoOperator
+{
+	One = 1,
+	Two = 2
+};
+
+enum struct ScopedEnumNoOperator
+{
+	One = 1,
+	Two = 2
+};
+
+enum UnscopedEnumWithOperator
+{
+	Three = 3,
+	Four = 4
+};
+
+CQTestConvert::PrintStream& operator<<(CQTestConvert::PrintStream& stream, const UnscopedEnumWithOperator& e)
+{
+	switch (e)
+	{
+		case UnscopedEnumWithOperator::Three:
+			stream << "Three";
+			break;
+		case UnscopedEnumWithOperator::Four:
+			stream << "Four";
+			break;
+		default:
+			stream << "Unknown";
+			break;
+	}
+	return stream;
+}
+
+enum struct ScopedEnumWithOperator
+{
+	One = 1,
+	Two = 2
+};
+
+CQTestConvert::PrintStream& operator<<(CQTestConvert::PrintStream& stream, const ScopedEnumWithOperator& e)
+{
+	switch (e)
+	{
+		case ScopedEnumWithOperator::One:
+			stream << "One";
+			break;
+		case ScopedEnumWithOperator::Two:
+			stream << "Two";
+			break;
+		default:
+			stream << "Unknown";
+			break;
+	}
+	return stream;
+}
+
+TEST_CLASS(TestConvertToString, "TestFramework.CQTest")
+{
+	TEST_METHOD(TypeWithExplicitOverride_ConvertToString_ReturnsOverride)
+	{
+		TypeWithExplicitOverride obj;
+		ASSERT_THAT(AreEqual(OverrideStr, CQTestConvert::ToString(obj)));
+	}
+
+	TEST_METHOD(TypeWithToString_ConvertToString_ReturnsToString)
+	{
+		TypeWithToString obj;
+		ASSERT_THAT(AreEqual(ToStringStr, CQTestConvert::ToString(obj)));
+	}
+
+	TEST_METHOD(TypeWithOstream_ConvertToString_ReturnsOStream)
+	{
+		TypeWithOStream obj;
+		ASSERT_THAT(AreEqual(FString(OStreamStr), CQTestConvert::ToString(obj)));
+	}
+
+	TEST_METHOD(TypeWithToStringAndOverride_ConvertToString_ReturnsOverride)
+	{
+		TypeWithToStringAndOverride obj;
+		ASSERT_THAT(AreEqual(OverrideStr, CQTestConvert::ToString(obj)));
+	}
+
+	TEST_METHOD(TypeWithToStringAndOStream_ConvertToString_ReturnsToString)
+	{
+		TypeWithToStringAndOStream obj;
+		ASSERT_THAT(AreEqual(ToStringStr, CQTestConvert::ToString(obj)));
+	}
+
+	TEST_METHOD(UnscopedEnumNoOperator_ConvertToString_ReturnsNumber)
+	{
+		auto e = UnscopedEnumNoOperator::One;
+		ASSERT_THAT(AreEqual(FString(TEXT("1")), CQTestConvert::ToString(e)));
+	}
+
+	TEST_METHOD(ScopedEnumNoOperator_ConvertToString_ReturnsNumber)
+	{
+		auto e = ScopedEnumNoOperator::Two;
+		ASSERT_THAT(AreEqual(FString(TEXT("2")), CQTestConvert::ToString(e)));
+	}
+
+	TEST_METHOD(UnscopedEnumWithOperator_ConvertToString_ReturnsString)
+	{
+		auto e = UnscopedEnumWithOperator::Three;
+		ASSERT_THAT(AreEqual(FString(TEXT("Three")), CQTestConvert::ToString(e)));
+	}
+
+	TEST_METHOD(ScopedEnumWithOperator_ConvertToString_ReturnsString)
+	{
+		auto e = ScopedEnumWithOperator::Two;
+		ASSERT_THAT(AreEqual(FString(TEXT("Two")), CQTestConvert::ToString(e)));
+	}
+
+	TEST_METHOD(ArrayOfValues_ConvertToString_ReturnsCommaDelimitedList)
+	{
+		TArray<int32> values = { 1, 2, 3 };
+		ASSERT_THAT(AreEqual(FString(TEXT("[1, 2, 3]")), CQTestConvert::ToString(values)));
+	}
+
+	TEST_METHOD(SetOfValues_ConvertToString_ReturnsCommaDelimitedList)
+	{
+		TSet<int32> values = { 1, 2, 3 };
+		auto asStr = CQTestConvert::ToString(values);
+		ASSERT_THAT(IsTrue(asStr.Contains(TEXT("1"))));
+		ASSERT_THAT(IsTrue(asStr.Contains(TEXT("2"))));
+		ASSERT_THAT(IsTrue(asStr.Contains(TEXT("3"))));
+	}
+
+	TEST_METHOD(MapOfValues_ConvertToString_ReturnsKeysAndValues)
+	{
+		TMap<int32, FString> values = { 
+			{ 1, FString(TEXT("One")) },
+			{ 2, FString(TEXT("Two")) },
+			{ 3, FString(TEXT("Three")) } 
+		};
+
+		auto asStr = CQTestConvert::ToString(values);
+		ASSERT_THAT(IsTrue(asStr.Contains(TEXT("1:One"))));
+		ASSERT_THAT(IsTrue(asStr.Contains(TEXT("2:Two"))));
+		ASSERT_THAT(IsTrue(asStr.Contains(TEXT("3:Three"))));
+	}
+};
diff --git a/Engine/Plugins/Tests/CQTest/Source/CQTestTests/Private/CQTestTests.cpp b/Engine/Plugins/Tests/CQTest/Source/CQTestTests/Private/CQTestTests.cpp
new file mode 100644
--- /dev/null
+++ b/Engine/Plugins/Tests/CQTest/Source/CQTestTests/Private/CQTestTests.cpp
@@ -0,0 +1,442 @@
+// Copyright Epic Games, Inc. All Rights Reserved.
+
+#include "CQTest.h"
+
+#include "Tickable.h"
+
+namespace CQTestTests
+{
+	TEST(MinimalTest, "TestFramework.CQTest")
+	{
+		ASSERT_THAT(IsTrue(true));
+	};
+
+	TEST_CLASS(From, "TestFramework.CQTest.GenerateDirectory.TokenProduces.[GenerateTestDirectory]")
+	{
+		TEST_METHOD(FolderStructure)
+		{
+			FString Expected = TEXT("CQTest");
+			ASSERT_THAT(IsTrue(TestRunner->TestDir.EndsWith(Expected), FString::Printf(TEXT("TestDir to end with %s but TestDir is %s produced from %s"), *Expected, *TestRunner->TestDir, *TestRunner->GetTestSourceFileName())));
+		}
+	};
+
+	TEST_CLASS(Produces, "TestFramework.CQTest.GenerateDirectory")
+	{
+		TEST_METHOD(WithPlugins_AppearsInPlugins)
+		{
+			FString GeneratedDirectory = TestDirectoryGenerator::Generate(TEXT("Projects/MyProject/Plugins/PluginOne/Source/Test.cpp"));
+			ASSERT_THAT(AreEqual(FString(TEXT("MyProject.Plugins.PluginOne")), GeneratedDirectory));
+		}
+
+		TEST_METHOD(WithPlatforms_AppearsInPlatforms)
+		{
+			FString GeneratedDirectory = TestDirectoryGenerator::Generate(TEXT("Projects/MyProject/Platforms/PlatformOne/Source/Test.cpp"));
+			ASSERT_THAT(AreEqual(FString(TEXT("MyProject.Platforms.PlatformOne")), GeneratedDirectory));
+		}
+
+		TEST_METHOD(WithoutPluginsOrPlatforms_FallsBackToSource)
+		{
+			FString GeneratedDirectory = TestDirectoryGenerator::Generate(TEXT("Project/MyProject/Source/MyProjectFolder/Test.cpp"));
+			ASSERT_THAT(AreEqual(FString(TEXT("MyProject.Source.MyProjectFolder")), GeneratedDirectory));
+		}
+	};
+
+	TEST_CLASS(SourceAndFile, "TestFramework.CQTest")
+	{
+		TEST_METHOD(SetsSourceFile)
+		{
+			ASSERT_THAT(AreEqual(FString(__FILE__), TestRunner->GetTestSourceFileName()));
+		}
+
+		TEST_METHOD(SetsLine_WithLineOfTestClass)
+		{
+			ASSERT_THAT(AreEqual(__LINE__ - 2, TestRunner->GetTestSourceFileLine(TEXT("SetsLine_WithLineOfTestClass"))));
+		}
+	};
+
+	TEST_CLASS(DefaultFixtureTestFlags, "TestFramework.CQTest")
+	{
+		TEST_METHOD(SetsApplicationContextMask)
+		{
+			ASSERT_THAT(AreEqual(EAutomationTestFlags::ApplicationContextMask, TestRunner->GetTestFlags() & EAutomationTestFlags::ApplicationContextMask));
+		}
+
+		TEST_METHOD(SetsProductFilter)
+		{
+			ASSERT_THAT(AreEqual(EAutomationTestFlags::ProductFilter, TestRunner->GetTestFlags() & EAutomationTestFlags::ProductFilter));
+		}
+	};
+
+	TEST_CLASS_WITH_FLAGS(OverrideFixtureTestFlags, "TestFramework.CQTest", EAutomationTestFlags::ApplicationContextMask | EAutomationTestFlags::SmokeFilter)
+	{
+		TEST_METHOD(GetTestFlags_ReturnsSetAutomationTestFlags)
+		{
+			ASSERT_THAT(AreEqual(EAutomationTestFlags::ApplicationContextMask | EAutomationTestFlags::SmokeFilter, TestRunner->GetTestFlags()));
+		}
+	};
+
+	TEST_CLASS(TestFixtureTest, "TestFramework.CQTest")
+	{
+		bool SetupCalled = false;
+		bool ShouldAddErrorDuringTearDown = false;
+		uint32 SomeNumber = 0;
+		FString ExpectedError = TEXT("Error reported in TearDown");
+
+		BEFORE_EACH()
+		{
+			SetupCalled = true;
+			SomeNumber++;
+		}
+
+		AFTER_EACH()
+		{
+			if (ShouldAddErrorDuringTearDown)
+			{
+				Assert.Fail(ExpectedError);
+			}
+		}
+
+	protected:
+		void ProtectedMethodDefinedInFixture() const {}
+
+		void AddExpectedErrorDuringTearDown()
+		{
+			ShouldAddErrorDuringTearDown = true;
+			Assert.ExpectError(ExpectedError);
+		}
+
+		TEST_METHOD(CanAccessProtectedFixtureMethods)
+		{
+			ProtectedMethodDefinedInFixture();
+		}
+
+		TEST_METHOD(BeforeRunTest_CallsSetup)
+		{
+			ASSERT_THAT(IsTrue(SetupCalled));
+		}
+
+		TEST_METHOD(AfterRunTest_CallsTearDown)
+		{
+			AddExpectedErrorDuringTearDown();
+		}
+
+		TEST_METHOD(BackingFixture_ResetsStateBetweenTestsPartOne)
+		{
+			ASSERT_THAT(AreEqual(1, SomeNumber));
+		}
+
+		TEST_METHOD(BackingFixture_ResetsStateBetweenTestsPartTwo)
+		{
+			ASSERT_THAT(AreEqual(1, SomeNumber));
+		}
+	};
+
+	TEST_CLASS(TestFixtureConstructor, "TestFramework.CQTest")
+	{
+		bool SetupCalled = false;
+		uint32 SomeNumber = 0;
+
+		TestFixtureConstructor()
+		{
+			SetupCalled = true;
+			SomeNumber++;
+		}
+
+	protected:
+		TEST_METHOD(ConstructorIsCalled_BeforeRunTest)
+		{
+			ASSERT_THAT(IsTrue(SetupCalled));
+		}
+
+		TEST_METHOD(BackingFixture_ResetsStateBetweenTestsPartOne)
+		{
+			ASSERT_THAT(AreEqual(1, SomeNumber));
+		}
+
+		TEST_METHOD(BackingFixture_ResetsStateBetweenTestsPartTwo)
+		{
+			ASSERT_THAT(AreEqual(1, SomeNumber));
+		}
+	};
+
+	static void ClearExpectedError(FAutomationTestBase& TestRunner, const FString& ExpectedError)
+	{
+		FAutomationTestExecutionInfo testInfo;
+		TestRunner.GetExecutionInfo(testInfo);
+		if (testInfo.GetErrorTotal() != 1)
+		{
+			return;
+		}
+		testInfo.RemoveAllEvents([&ExpectedError](FAutomationEvent& event) {
+			return event.Message.Equals(ExpectedError);
+		});
+		if (testInfo.GetErrorTotal() == 0)
+		{
+			TestRunner.ClearExecutionInfo();
+		}
+	}
+	//This test is verifying that if an assertion is raised in BEFORE_EACH then the TEST_METHOD does not run
+	//But the AFTER_EACH still is run.  This is done by Asserting with an expected message in BEFORE and Asserting
+	//with an unexpected message in TEST_METHOD.  The AFTER_EACH checks for the expected message, and clears the errors
+	//if an only if there was 1 error message with the expected value
+
+	TEST_CLASS(TestAssertionInBefore, "TestFramework.CQTest")
+	{
+		FString ExpectedError = TEXT("Expected Error Message");
+
+		BEFORE_EACH() 
+		{
+			Assert.Fail(ExpectedError);
+		}
+
+		AFTER_EACH() 
+		{
+			ClearExpectedError(*this->TestRunner, ExpectedError);
+		}
+
+		TEST_METHOD(BeforeTest_AssertionFailure_DoesNotRunTestMethod)
+		{
+			Assert.Fail(TEXT("TEST_METHOD should not run if assertion fails in BEFORE_EACH"));
+		}
+	};
+
+	// --------------------------------------------------------
+	// Latent commands are awaited
+	// --------------------------------------------------------
+	template <typename Test>
+	class FMinimumCallCommand : public IAutomationLatentCommand
+	{
+	public:
+		FMinimumCallCommand(Test* InTest, int32 expectedCount)
+			: ExecutingTest(InTest), ExpectedCount(expectedCount) {}
+
+		bool Update() override
+		{
+			if (ExecutingTest)
+			{
+				ExecutingTest->IncrementExecutedCommandsCount();
+			}
+			CurrentCount++;
+			return CurrentCount == ExpectedCount;
+		}
+
+		Test* ExecutingTest;
+		int32 ExpectedCount{ 0 };
+		int32 CurrentCount{ 0 };
+	};
+
+	TEST_CLASS(LatentCommandTest, "TestFramework.CQTest")
+	{
+		void SetExpectedExecutedCommandsCount(int32 count)
+		{
+			ExpectedExecutedCommandsCount = count;
+		}
+		void IncrementExecutedCommandsCount()
+		{
+			ExecutedCommandsCount++;
+		}
+
+		int32 ExpectedExecutedCommandsCount{ 0 };
+		int32 ExecutedCommandsCount{ 0 };
+
+		BEFORE_EACH()
+		{
+			for (int32 i = 0; i < 3; i++)
+			{
+				AddCommand(new FMinimumCallCommand(this, i + 1));
+			}
+		}
+
+		AFTER_EACH()
+		{
+			if (ExpectedExecutedCommandsCount > 0)
+			{
+				ASSERT_THAT(AreEqual(ExpectedExecutedCommandsCount, ExecutedCommandsCount));
+			}
+		}
+
+		TEST_METHOD(Test_WithCommandsInBeforeTest_ExecutesCommandsBeforeRun)
+		{
+			ASSERT_THAT(AreEqual(1 + 2 + 3, ExecutedCommandsCount));
+		}
+
+		TEST_METHOD(Test_WithLatentCommandsInTest_ExecutesCommandsBeforeTearDown)
+		{
+			ExpectedExecutedCommandsCount = ExecutedCommandsCount * 2;
+			for (int32 i = 0; i < 3; i++)
+			{
+				AddCommand(new FMinimumCallCommand(this, i + 1));
+			}
+		}
+	};
+
+	// --------------------------------------------------------
+	// Tickable Game Objects Tick
+	// --------------------------------------------------------
+	struct FTestTickable : public FTickableGameObject
+	{
+		virtual TStatId GetStatId() const override
+		{
+			return TStatId();
+		}
+
+		virtual void Tick(float DeltaTime) override
+		{
+			TickCount++;
+		}
+
+		virtual bool IsTickableInEditor() const override
+		{
+			return true;
+		}
+
+		void ResetTickCount()
+		{
+			TickCount = 0;
+		}
+
+		uint32 TickCount = 0;
+	};
+
+	TEST_CLASS(GameObjectsTickTest, "TestFramework.CQTest")
+	{
+		BEFORE_EACH()
+		{
+			Tickable.ResetTickCount();
+			AddCommand(new FWaitUntil(TestRunner, [&]() { return Tickable.TickCount > 2; }));
+		}
+		AFTER_EACH()
+		{
+			ASSERT_THAT(IsTrue(Tickable.TickCount > 2));
+		}
+
+		FTestTickable Tickable{};
+
+		TEST_METHOD(TestWithTickableGameObject_WaitingForTicksInSetup_WillAllowGameObjectToTick)
+		{
+			TestCommandBuilder.Do([this]() { ASSERT_THAT(IsTrue(Tickable.TickCount > 2)); });
+		}
+
+		TEST_METHOD(TestWithTickableGameObject_WaitingForTicksInSetup_WillBeCompleteDuringRunStep)
+		{
+			ASSERT_THAT(IsTrue(Tickable.TickCount > 2));
+		}
+	};
+
+	template <typename Test>
+	class FDelayedCommand : public IAutomationLatentCommand
+	{
+	public:
+		explicit FDelayedCommand(Test* InTest, int32 RequiredTicks)
+			: ExecutingTest(InTest) 
+			, RemainingTicks(RequiredTicks)
+		{}
+
+		bool Update() override
+		{
+			RemainingTicks--;
+			if (RemainingTicks == 0)
+			{
+				ExecutingTest->DelayedCommandExecuted = true;
+				return true;			
+			}
+
+			return false;
+		}
+
+		Test* ExecutingTest;
+		int32 RemainingTicks;
+	};
+
+	template <typename Test>
+	class FQueueCommand : public IAutomationLatentCommand
+	{
+	public:
+		explicit FQueueCommand(Test* InTest, int32 RequiredTicks)
+			: ExecutingTest(InTest) 
+			, RemainingTicks(RequiredTicks)
+			, RequiredTicks(RequiredTicks)
+		{}
+
+		bool Update() override
+		{
+			RemainingTicks--;
+			if (RemainingTicks == 0) {
+				ExecutingTest->AddCommand(new FDelayedCommand(ExecutingTest, RequiredTicks));
+				return true;
+			}
+			return false;
+		}
+
+		Test* ExecutingTest;
+		int32 RemainingTicks;
+		int32 RequiredTicks;
+	};
+
+	TEST_CLASS(ChainedCommandsTest, "TestFramework.CQTest")
+	{
+		bool DelayedCommandExecuted{ false };
+
+		BEFORE_EACH()
+		{
+			AddCommand(new FQueueCommand(this, 1));
+		}
+
+		AFTER_EACH()
+		{
+			ASSERT_THAT(IsTrue(DelayedCommandExecuted));
+
+			DelayedCommandExecuted = false;
+			TSharedPtr<FQueueCommand<ChainedCommandsTest>> PostAfterEach = MakeShared<FQueueCommand<ChainedCommandsTest>>(this, 1);
+			TSharedPtr<FExecute> Assertion = MakeShared<FExecute>([&]() {
+				ASSERT_THAT(IsTrue(DelayedCommandExecuted));
+			});
+			AddCommand(new FRunSequence(PostAfterEach, Assertion));
+		}
+
+		TEST_METHOD(DelayedCommandExecuted_WhenAddedInBeforeEach_IsTrueDuringRunTest)
+		{
+			ASSERT_THAT(IsTrue(DelayedCommandExecuted));
+		}
+
+		TEST_METHOD(DelayedCommandExecuted_WhenAddedDuringTest_IsTrueAfterTest)
+		{
+			DelayedCommandExecuted = false;
+			AddCommand(new FQueueCommand(this, 1));
+		}
+	};
+
+	TEST_CLASS(MultiTickChainedCommandsTest, "TestFramework.CQTest")
+	{
+		bool DelayedCommandExecuted{ false };
+
+		BEFORE_EACH()
+		{
+			AddCommand(new FQueueCommand(this, 3));
+		}
+
+		AFTER_EACH()
+		{
+			ASSERT_THAT(IsTrue(DelayedCommandExecuted));
+			DelayedCommandExecuted = false;
+
+			TSharedPtr<FQueueCommand<MultiTickChainedCommandsTest>> PostAfterEach = MakeShared<FQueueCommand<MultiTickChainedCommandsTest>>(this, 3);
+			TSharedPtr<FExecute> Assertion = MakeShared<FExecute>([&]() {
+				ASSERT_THAT(IsTrue(DelayedCommandExecuted));
+			});
+
+			AddCommand(new FRunSequence(PostAfterEach, Assertion));
+		}
+
+		TEST_METHOD(DelayedCommandExecuted_WhenAddedInBeforeEach_IsTrueDuringRunTest)
+		{
+			ASSERT_THAT(IsTrue(DelayedCommandExecuted));
+		}
+
+		TEST_METHOD(DelayedCommandExecuted_WhenAddedDuringTest_IsTrueAfterTest)
+		{
+			DelayedCommandExecuted = false;
+			AddCommand(new FQueueCommand(this, 3));
+		}
+	};
+} // namespace CQTestTests
diff --git a/Engine/Plugins/Tests/CQTest/Source/CQTestTests/Private/CommandBuilderTests.cpp b/Engine/Plugins/Tests/CQTest/Source/CQTestTests/Private/CommandBuilderTests.cpp
new file mode 100644
--- /dev/null
+++ b/Engine/Plugins/Tests/CQTest/Source/CQTestTests/Private/CommandBuilderTests.cpp
@@ -0,0 +1,42 @@
+// Copyright Epic Games, Inc. All Rights Reserved.
+
+#include "CQTest.h"
+
+TEST_CLASS(CommandBuilderTests, "TestFramework.CQTest")
+{
+	FTestCommandBuilder CommandBuilder{};
+
+	TEST_METHOD(Do_ThenBuild_IncludesCommand)
+	{
+		bool invoked = false;
+		auto command = CommandBuilder.Do([&invoked]() { invoked = true; }).Build();
+
+		ASSERT_THAT(IsTrue(command->Update()));
+		ASSERT_THAT(IsTrue(invoked));
+	}
+
+	TEST_METHOD(Build_WithoutCommands_ReturnsNullptr)
+	{
+		auto command = CommandBuilder.Build();
+		ASSERT_THAT(IsNull(command));
+	}
+
+	TEST_METHOD(StartWhen_CreatesWaitUntilCommand)
+	{
+		bool done = false;
+		auto command = CommandBuilder.StartWhen(TestRunner, [&done]() { return done; }).Build();
+
+		ASSERT_THAT(IsFalse(command->Update()));
+		done = true;
+		ASSERT_THAT(IsTrue(command->Update()));
+	}
+
+	TEST_METHOD(Build_AfterBuild_ReturnsNullptr)
+	{
+		auto command = CommandBuilder.Do([]() {}).Build();
+		auto secondTime = CommandBuilder.Build();
+
+		ASSERT_THAT(IsNotNull(command));
+		ASSERT_THAT(IsNull(secondTime));
+	}
+};
\ No newline at end of file
diff --git a/Engine/Plugins/Tests/CQTest/Source/CQTestTests/Private/ExtensionTests.cpp b/Engine/Plugins/Tests/CQTest/Source/CQTestTests/Private/ExtensionTests.cpp
new file mode 100644
--- /dev/null
+++ b/Engine/Plugins/Tests/CQTest/Source/CQTestTests/Private/ExtensionTests.cpp
@@ -0,0 +1,74 @@
+// Copyright Epic Games, Inc. All Rights Reserved.
+
+#include "CQTest.h"
+#include "Assert/NoDiscardAsserter.h"
+
+struct FCustomAsserter : public FNoDiscardAsserter
+{
+public:
+	FCustomAsserter(FAutomationTestBase& testRunner)
+		: FNoDiscardAsserter(testRunner)
+	{
+	}
+
+	void Custom(bool) {
+	}
+};
+
+#define CUSTOM_ASSERT_TEST_CLASS(_ClassName, _TestDir) TEST_CLASS_WITH_ASSERTS(_ClassName, _TestDir, FCustomAsserter)
+
+CUSTOM_ASSERT_TEST_CLASS(CustomAsserts, "TestFramework.CQTest")
+{
+	TEST_METHOD(CustomTestClass_WithCustomAsserter_HasInstanceOfCustomAsserter)
+	{
+		Assert.Custom(true);
+		ASSERT_THAT(IsTrue(true));
+	}
+};
+
+template<typename Derived, typename AsserterType>
+struct TCustomBaseClass : public TTest<Derived, AsserterType>
+{
+	uint32 SpecialValue = 42;
+};
+
+#define CUSTOM_BASE_TEST_CLASS(_ClassName, _TestDir) TEST_CLASS_WITH_BASE(_ClassName, _TestDir, TCustomBaseClass)
+
+CUSTOM_BASE_TEST_CLASS(CustomBase, "TestFramework.CQTest")
+{
+	TEST_METHOD(CustomTestClass_WithCustomBase_InheritsFromBaseClass)
+	{
+		ASSERT_THAT(AreEqual(42, SpecialValue));
+	}
+};
+
+template <typename Derived, typename AsserterType>
+struct TBaseWithConstructor : public TTest<Derived, AsserterType>
+{
+	static inline bool bInitializedCall{ false };
+
+	static inline bool bNonInitializedCall{ false };
+
+	TBaseWithConstructor()
+	{
+		if (this->bInitializing)
+		{
+			bInitializedCall = true;
+		}
+		else
+		{
+			bNonInitializedCall = true;
+		}
+	}
+};
+
+#define CUSTOM_WITH_CTOR_CLASS(_ClassName, _TestDir) TEST_CLASS_WITH_BASE(_ClassName, _TestDir, TBaseWithConstructor)
+
+CUSTOM_WITH_CTOR_CLASS(CustomCtor, "TestFramework.CQTest")
+{
+	TEST_METHOD(CustomClassWithCtor_CallsCtorWhenInitializing_AndBeforeEachTest)
+	{
+		ASSERT_THAT(IsTrue(bInitializedCall));	 // called when populating test names
+		ASSERT_THAT(IsTrue(bNonInitializedCall)); // called when creating this test
+	}
+};
\ No newline at end of file
diff --git a/Engine/Plugins/Tests/CQTest/Source/CQTestTests/Private/ObjectBuilderTests.cpp b/Engine/Plugins/Tests/CQTest/Source/CQTestTests/Private/ObjectBuilderTests.cpp
new file mode 100644
--- /dev/null
+++ b/Engine/Plugins/Tests/CQTest/Source/CQTestTests/Private/ObjectBuilderTests.cpp
@@ -0,0 +1,1157 @@
+// Copyright Epic Games, Inc. All Rights Reserved.
+
+#include "TestActorWithProperties.h"
+
+#include "Animation/AnimSequence.h"
+#include "Components/SceneComponent.h"
+
+#include "ObjectBuilder.h"
+#include "ActorTestSpawner.h"
+#include "CQTest.h"
+
+TEST_CLASS(ObjectBuilder_Success, "TestFramework.CQTest")
+{
+	FActorTestSpawner Spawner;
+
+	TEST_METHOD(Spawn_NonActorWithProperties_SetProperties)
+	{
+		auto& Result = TObjectBuilder<UAnimSequence>()
+						   .SetParam<bool>("bEnableRootMotion", true)
+						   .SetParam("RefFrameIndex", 321)
+						   .Spawn();
+
+		ASSERT_THAT(IsTrue(Result.bEnableRootMotion));
+		ASSERT_THAT(AreEqual(Result.RefFrameIndex, 321));
+	}
+	TEST_METHOD(Spawn_ActorWithBool_SetsBooleanValue)
+	{
+		auto& Result = TObjectBuilder<ATestActorWithProperties>(Spawner).SetParam("BoolProperty", true).Spawn();
+		ASSERT_THAT(IsTrue(Result.BoolProperty));
+		auto& ConstResult = TObjectBuilder<ATestActorWithProperties>(Spawner).SetParam<const bool>("BoolProperty", true).Spawn();
+		ASSERT_THAT(IsTrue(ConstResult.BoolProperty));
+	}
+	TEST_METHOD(Spawn_ActorWithByte_SetsByteValue)
+	{
+		uint8 Byte = 250;
+		auto& Result = TObjectBuilder<ATestActorWithProperties>(Spawner).SetParam("ByteProperty", Byte).Spawn();
+		ASSERT_THAT(AreEqual(Result.ByteProperty, Byte));
+		auto& ConstResult = TObjectBuilder<ATestActorWithProperties>(Spawner).SetParam<const uint8>("ByteProperty", Byte).Spawn();
+		ASSERT_THAT(AreEqual(ConstResult.ByteProperty, Byte));
+	}
+
+	TEST_METHOD(Spawn_ActorWithIntegers_SetsIntegers)
+	{
+		int8 int8Value = -112;
+		int16 int16Value = 10000;
+		int32 int32Value = -1000000;
+		int64 int64Value = -10000000;
+		auto& Result = TObjectBuilder<ATestActorWithProperties>(Spawner)
+						   .SetParam("Int8Property", int8Value)
+						   .SetParam("Int16Property", int16Value)
+						   .SetParam("Int32Property", int32Value)
+						   .SetParam("Int64Property", int64Value)
+						   .Spawn();
+
+		ASSERT_THAT(AreEqual(Result.Int8Property, int8Value));
+		ASSERT_THAT(AreEqual(Result.Int16Property, int16Value));
+		ASSERT_THAT(AreEqual(Result.Int32Property, int32Value));
+		ASSERT_THAT(AreEqual(Result.Int64Property, int64Value));
+
+		auto& ConstResult = TObjectBuilder<ATestActorWithProperties>(Spawner)
+								.SetParam<const int8>("Int8Property", int8Value)
+								.SetParam<const int16>("Int16Property", int16Value)
+								.SetParam<const int32>("Int32Property", int32Value)
+								.SetParam<const int64>("Int64Property", int64Value)
+								.Spawn();
+
+		ASSERT_THAT(AreEqual(ConstResult.Int8Property, int8Value));
+		ASSERT_THAT(AreEqual(ConstResult.Int16Property, int16Value));
+		ASSERT_THAT(AreEqual(ConstResult.Int32Property, int32Value));
+		ASSERT_THAT(AreEqual(ConstResult.Int64Property, int64Value));
+	}
+
+	TEST_METHOD(Spawn_ActorWithUnsignedIntegers_SetsUnsignedIntegers)
+	{
+		uint16 uint16Value = 10000;
+		uint32 uint32Value = 5000000;
+		uint64 uint64Value = 400000000;
+
+		auto& Result = TObjectBuilder<ATestActorWithProperties>(Spawner)
+						   .SetParam("UInt16Property", uint16Value)
+						   .SetParam("UInt32Property", uint32Value)
+						   .SetParam("UInt64Property", uint64Value)
+						   .Spawn();
+
+		ASSERT_THAT(AreEqual(Result.UInt16Property, uint16Value));
+		ASSERT_THAT(AreEqual(Result.UInt32Property, uint32Value));
+		ASSERT_THAT(AreEqual(Result.UInt64Property, uint64Value));
+
+		auto& ConstResult = TObjectBuilder<ATestActorWithProperties>(Spawner)
+						   .SetParam<const uint16>("UInt16Property", uint16Value)
+						   .SetParam<const uint32>("UInt32Property", uint32Value)
+						   .SetParam<const uint64>("UInt64Property", uint64Value)
+						   .Spawn();
+
+		ASSERT_THAT(AreEqual(ConstResult.UInt16Property, uint16Value));
+		ASSERT_THAT(AreEqual(ConstResult.UInt32Property, uint32Value));
+		ASSERT_THAT(AreEqual(ConstResult.UInt64Property, uint64Value));
+	}
+
+	TEST_METHOD(Spawn_ActorWithFloats_SetsFloats)
+	{
+		float floatValue = 1.23f;
+		double doubleValue = 1.23456789;
+
+		auto& Result = TObjectBuilder<ATestActorWithProperties>(Spawner)
+						   .SetParam("FloatProperty", floatValue)
+						   .SetParam("DoubleProperty", doubleValue)
+						   .Spawn();
+
+		ASSERT_THAT(AreEqual(Result.FloatProperty, floatValue));
+		ASSERT_THAT(AreEqual(Result.DoubleProperty, doubleValue));
+
+		auto& ConstResult = TObjectBuilder<ATestActorWithProperties>(Spawner)
+						   .SetParam<const float>("FloatProperty", floatValue)
+						   .SetParam<const double>("DoubleProperty", doubleValue)
+						   .Spawn();
+
+		ASSERT_THAT(AreEqual(ConstResult.FloatProperty, floatValue));
+		ASSERT_THAT(AreEqual(ConstResult.DoubleProperty, doubleValue));
+	}
+
+	TEST_METHOD(Spawn_ActorWithEnums_SetsEnums)
+	{
+		ETestInt8 s8Enum = ETestInt8::enumone;
+		ETestInt16 s16Enum = ETestInt16::enumtwo;
+		ETestInt32 s32Enum = ETestInt32::enumthree;
+		ETestInt64 s64Enum = ETestInt64::enumtwo;
+
+		ETestUint8 u8Enum = ETestUint8::enumthree;
+		ETestUint16 u16Enum = ETestUint16::enumone;
+		ETestUint32 u32Enum = ETestUint32::enumtwo;
+		ETestUint64 u64Enum = ETestUint64::enumone;
+
+		auto& Result = TObjectBuilder<ATestActorWithProperties>(Spawner)
+						   .SetParam("Int8EnumProperty", s8Enum)
+						   .SetParam("Int16EnumProperty", s16Enum)
+						   .SetParam("Int32EnumProperty", s32Enum)
+						   .SetParam("Int64EnumProperty", s64Enum)
+						   .SetParam("Uint8EnumProperty", u8Enum)
+						   .SetParam("Uint16EnumProperty", u16Enum)
+						   .SetParam("Uint32EnumProperty", u32Enum)
+						   .SetParam("Uint64EnumProperty", u64Enum)
+						   .Spawn();
+
+		ASSERT_THAT(AreEqual(Result.Int8EnumProperty, s8Enum));
+		ASSERT_THAT(AreEqual(Result.Int16EnumProperty, s16Enum));
+		ASSERT_THAT(AreEqual(Result.Int32EnumProperty, s32Enum));
+		ASSERT_THAT(AreEqual(Result.Int64EnumProperty, s64Enum));
+		ASSERT_THAT(AreEqual(Result.Uint8EnumProperty, u8Enum));
+		ASSERT_THAT(AreEqual(Result.Uint16EnumProperty, u16Enum));
+		ASSERT_THAT(AreEqual(Result.Uint32EnumProperty, u32Enum));
+		ASSERT_THAT(AreEqual(Result.Uint64EnumProperty, u64Enum));
+
+		auto& ConstResult = TObjectBuilder<ATestActorWithProperties>(Spawner)
+						   .SetParam<const ETestInt8>("Int8EnumProperty", s8Enum)
+						   .SetParam<const ETestInt16>("Int16EnumProperty", s16Enum)
+						   .SetParam<const ETestInt32>("Int32EnumProperty", s32Enum)
+						   .SetParam<const ETestInt64>("Int64EnumProperty", s64Enum)
+						   .SetParam<const ETestUint8>("Uint8EnumProperty", u8Enum)
+						   .SetParam<const ETestUint16>("Uint16EnumProperty", u16Enum)
+						   .SetParam<const ETestUint32>("Uint32EnumProperty", u32Enum)
+						   .SetParam<const ETestUint64>("Uint64EnumProperty", u64Enum)
+						   .Spawn();
+
+		ASSERT_THAT(AreEqual(ConstResult.Int8EnumProperty, s8Enum));
+		ASSERT_THAT(AreEqual(ConstResult.Int16EnumProperty, s16Enum));
+		ASSERT_THAT(AreEqual(ConstResult.Int32EnumProperty, s32Enum));
+		ASSERT_THAT(AreEqual(ConstResult.Int64EnumProperty, s64Enum));
+		ASSERT_THAT(AreEqual(ConstResult.Uint8EnumProperty, u8Enum));
+		ASSERT_THAT(AreEqual(ConstResult.Uint16EnumProperty, u16Enum));
+		ASSERT_THAT(AreEqual(ConstResult.Uint32EnumProperty, u32Enum));
+		ASSERT_THAT(AreEqual(ConstResult.Uint64EnumProperty, u64Enum));
+	}
+
+	TEST_METHOD(Spawn_WithFName_SetsFName)
+	{
+		FName Name = FName(TEXT("TestName"));
+
+		auto& Result = TObjectBuilder<ATestActorWithProperties>(Spawner)
+						   .SetParam("NameProperty", Name)
+						   .Spawn();
+
+		ASSERT_THAT(AreEqual(Result.NameProperty, Name));
+
+		auto& ConstResult = TObjectBuilder<ATestActorWithProperties>(Spawner)
+						   .SetParam<const FName>("NameProperty", Name)
+						   .Spawn();
+
+		ASSERT_THAT(AreEqual(ConstResult.NameProperty, Name));
+
+	}
+
+	TEST_METHOD(Spawn_WithTObjectPtrProperty_SetsProperty)
+	{
+		auto& BuiltObject = TObjectBuilder<UAnimSequence>().Spawn();
+		auto ptr = TObjectPtr<UAnimSequence>(&BuiltObject);
+
+		auto& Result = TObjectBuilder<ATestActorWithProperties>(Spawner)
+						   .SetParam("TestTObjectPtrProperty", ptr)
+						   .Spawn();
+
+		ASSERT_THAT(IsTrue(Result.TestTObjectPtrProperty == ptr));
+
+		auto& ConstResult = TObjectBuilder<ATestActorWithProperties>(Spawner)
+						   .SetParam<const TObjectPtr<UAnimSequence>>("TestTObjectPtrProperty", ptr)
+						   .Spawn();
+
+		ASSERT_THAT(IsTrue(ConstResult.TestTObjectPtrProperty == ptr));
+	}
+
+	TEST_METHOD(Spawn_WithVectorProperty_SetsProperty)
+	{
+		auto& Result = TObjectBuilder<ATestActorWithProperties>(Spawner)
+						   .SetParam("VectorProperty", FVector::UpVector)
+						   .Spawn();
+
+		ASSERT_THAT(AreEqual(Result.VectorProperty, FVector::UpVector));
+
+
+		auto& ConstResult = TObjectBuilder<ATestActorWithProperties>(Spawner)
+						   .SetParam<const FVector>("VectorProperty", FVector::UpVector)
+						   .Spawn();
+
+		ASSERT_THAT(AreEqual(ConstResult.VectorProperty, FVector::UpVector));
+	}
+
+	TEST_METHOD(Spawn_WithStructProperty_SetsProperty)
+	{
+		TArray<int32> ArrayInStruct;
+		ArrayInStruct.Add(123);
+		ArrayInStruct.Add(456);
+		ArrayInStruct.Add(789);
+
+		FTestStructWithProperties UnrealStruct;
+		UnrealStruct.StructInt32Property = 321;
+		UnrealStruct.StructArrayProperty = ArrayInStruct;
+
+		auto& Result = TObjectBuilder<ATestActorWithProperties>(Spawner)
+						   .SetParam("StructProperty", UnrealStruct)
+						   .Spawn();
+
+		ASSERT_THAT(AreEqual(Result.StructProperty.StructInt32Property, 321));
+		ASSERT_THAT(AreEqual(Result.StructProperty.StructArrayProperty, ArrayInStruct));
+
+
+		auto& ConstResult = TObjectBuilder<ATestActorWithProperties>(Spawner)
+						   .SetParam<const FTestStructWithProperties>("StructProperty", UnrealStruct)
+						   .Spawn();
+
+		ASSERT_THAT(AreEqual(ConstResult.StructProperty.StructInt32Property, 321));
+		ASSERT_THAT(AreEqual(ConstResult.StructProperty.StructArrayProperty, ArrayInStruct));
+	}
+
+	TEST_METHOD(Spawn_WithPrimitiveTArray_SetsArray)
+	{
+		TArray<int32> Array;
+		Array.Add(1);
+		Array.Add(2);
+		Array.Add(3);
+		Array.Add(4);
+
+		auto& Result = TObjectBuilder<ATestActorWithProperties>(Spawner)
+						   .SetParam("ArrayProperty", Array)
+						   .Spawn();
+
+		ASSERT_THAT(IsTrue(Result.ArrayProperty == Array));
+
+		auto& ConstResult = TObjectBuilder<ATestActorWithProperties>(Spawner)
+						   .SetParam<const TArray<int32>>("ArrayProperty", Array)
+						   .Spawn();
+
+		ASSERT_THAT(IsTrue(ConstResult.ArrayProperty == Array));
+	}
+
+	TEST_METHOD(Spawn_WithObjectArray_SetsArray)
+	{
+		TArray<TObjectPtr<ATestActorWithProperties>> ObjectArray;
+		ObjectArray.Add(&TObjectBuilder<ATestActorWithProperties>(Spawner).Spawn());
+
+		auto& Result = TObjectBuilder<ATestActorWithProperties>(Spawner)
+						   .SetParam("ArrayOfObjectsProperty", ObjectArray)
+						   .Spawn();
+
+		ASSERT_THAT(IsTrue(Result.ArrayOfObjectsProperty == ObjectArray));
+
+		auto& ConstResult = TObjectBuilder<ATestActorWithProperties>(Spawner)
+						   .SetParam<const TArray<ATestActorWithProperties*>>("ArrayOfObjectsProperty", ObjectArray)
+						   .Spawn();
+
+		ASSERT_THAT(IsTrue(ConstResult.ArrayOfObjectsProperty == ObjectArray));
+	}
+
+	TEST_METHOD(Spawn_WithStructArray_SetsProperty)
+	{
+		TArray<int32> ArrayInStruct;
+		ArrayInStruct.Add(123);
+		ArrayInStruct.Add(456);
+		ArrayInStruct.Add(789);
+
+		FTestStructWithProperties UnrealStruct;
+		UnrealStruct.StructInt32Property = 321;
+		UnrealStruct.StructArrayProperty = ArrayInStruct;
+
+		TArray<int32> Array;
+		Array.Add(1);
+		Array.Add(2);
+		Array.Add(3);
+		Array.Add(4);
+
+		FTestStructWithProperties OtherUnrealStruct;
+		OtherUnrealStruct.StructInt32Property = 426;
+		OtherUnrealStruct.StructArrayProperty = Array;
+
+		TArray<FTestStructWithProperties> StructArray;
+		StructArray.Add(UnrealStruct);
+		StructArray.Add(OtherUnrealStruct);
+
+		auto& Result = TObjectBuilder<ATestActorWithProperties>(Spawner)
+						   .SetParam("ArrayOfStructsProperty", StructArray)
+						   .Spawn();
+
+		ASSERT_THAT(AreEqual(Result.ArrayOfStructsProperty.Num(), 2));
+		ASSERT_THAT(IsTrue(Result.ArrayOfStructsProperty[0].StructInt32Property == 426 || Result.ArrayOfStructsProperty[1].StructInt32Property == 426));
+		ASSERT_THAT(IsTrue(Result.ArrayOfStructsProperty[0].StructInt32Property == 321 || Result.ArrayOfStructsProperty[1].StructInt32Property == 321));
+		ASSERT_THAT(AreNotEqual(Result.ArrayOfStructsProperty[0].StructInt32Property, Result.ArrayOfStructsProperty[1].StructInt32Property));
+
+
+		auto& ConstResult = TObjectBuilder<ATestActorWithProperties>(Spawner)
+						   .SetParam<const TArray<FTestStructWithProperties>>("ArrayOfStructsProperty", StructArray)
+						   .Spawn();
+
+		ASSERT_THAT(AreEqual(ConstResult.ArrayOfStructsProperty.Num(), 2));
+		ASSERT_THAT(IsTrue(ConstResult.ArrayOfStructsProperty[0].StructInt32Property == 426 || Result.ArrayOfStructsProperty[1].StructInt32Property == 426));
+		ASSERT_THAT(IsTrue(ConstResult.ArrayOfStructsProperty[0].StructInt32Property == 321 || Result.ArrayOfStructsProperty[1].StructInt32Property == 321));
+		ASSERT_THAT(AreNotEqual(ConstResult.ArrayOfStructsProperty[0].StructInt32Property, Result.ArrayOfStructsProperty[1].StructInt32Property));
+	}
+
+	TEST_METHOD(Spawn_WithVectorArray_SetsArray)
+	{
+		TArray<FVector> VectorArray;
+		VectorArray.Add(FVector::UpVector);
+		VectorArray.Add(FVector::LeftVector);
+
+		auto& Result = TObjectBuilder<ATestActorWithProperties>(Spawner)
+						   .SetParam("ArrayOfVectorsProperty", VectorArray)
+						   .Spawn();
+
+		ASSERT_THAT(IsTrue(Result.ArrayOfVectorsProperty == VectorArray));
+
+
+		auto& ConstResult = TObjectBuilder<ATestActorWithProperties>(Spawner)
+						   .SetParam<const TArray<FVector>>("ArrayOfVectorsProperty", VectorArray)
+						   .Spawn();
+
+		ASSERT_THAT(IsTrue(ConstResult.ArrayOfVectorsProperty == VectorArray));
+	}
+
+	TEST_METHOD(Spawn_WithSet_SetsSet)
+	{
+		TSet<int32> Set;
+		Set.Add(1);
+
+		auto& Result = TObjectBuilder<ATestActorWithProperties>(Spawner)
+						   .SetParam("SetProperty", Set)
+						   .Spawn();
+
+		ASSERT_THAT(IsNotNull(Result.SetProperty.Find(1)));
+
+		auto& ConstResult = TObjectBuilder<ATestActorWithProperties>(Spawner)
+						   .SetParam<const TSet<int32>>("SetProperty", Set)
+						   .Spawn();
+
+		ASSERT_THAT(IsNotNull(ConstResult.SetProperty.Find(1)));
+	}
+
+	TEST_METHOD(Spawn_WithMapProperty_SetsMap)
+	{
+		TMap<int32, int32> Map;
+		Map.Add(1, 2);
+		Map.Add(3, 4);
+		Map.Add(5, 6);
+
+		auto& Result = TObjectBuilder<ATestActorWithProperties>(Spawner)
+						   .SetParam("MapProperty", Map)
+						   .Spawn();
+
+		ASSERT_THAT(IsTrue(Result.MapProperty.OrderIndependentCompareEqual(Map)));
+
+		auto& ConstResult = TObjectBuilder<ATestActorWithProperties>(Spawner)
+						   .SetParam<const TMap<int32, int32>>("MapProperty", Map)
+						   .Spawn();
+
+		ASSERT_THAT(IsTrue(ConstResult.MapProperty.OrderIndependentCompareEqual(Map)));
+	}
+
+	TEST_METHOD(Spawn_WithComponentType_AddsComponent)
+	{
+		auto& Result = TObjectBuilder<ATestActorWithProperties>(Spawner)
+						   .AddComponentTo<USceneComponent>()
+						   .Spawn();
+
+		ASSERT_THAT(AreEqual(Result.GetComponents().Num(), 1));
+	}
+
+	TEST_METHOD(Spawn_WithBuiltComponent_AddsComponentInPreBuiltState)
+	{
+		const FVector Location(1.0f, 2.0f, 3.0f);
+		auto& BuiltComponent = TObjectBuilder<USceneComponent>()
+								   .SetParam("RelativeLocation", Location)
+								   .Spawn();
+
+		auto& Result = TObjectBuilder<ATestActorWithProperties>(Spawner)
+						   .AddComponentTo(&BuiltComponent)
+						   .Spawn();
+
+		auto Index = Result.GetComponents().FindId(&BuiltComponent);
+		ASSERT_THAT(IsTrue(Index.IsValidId()));
+
+		auto* Component = Cast<USceneComponent>(Result.GetComponents()[Index]);
+		ASSERT_THAT(IsNotNull(Component));
+		ASSERT_THAT(AreEqual(Component->GetRelativeLocation(), Location));
+	}
+
+	TEST_METHOD(Spawn_WithChildActor_AddsChildActor)
+	{
+		auto& Result = TObjectBuilder<ATestActorWithProperties>(Spawner)
+						   .AddChildActorComponentTo<AActor>()
+						   .Spawn();
+
+		ASSERT_THAT(AreEqual(Result.GetComponents().Num(), 1));
+	}
+
+	TEST_METHOD(SpawnObject_WithNoOwner_MatchesDefaultOwner)
+	{
+		auto* NormalObject = NewObject<USceneComponent>();
+		auto& BuiltObject = TObjectBuilder<USceneComponent>().Spawn();
+
+		ASSERT_THAT(AreEqual(NormalObject->GetOwner(), BuiltObject.GetOwner()));
+	}
+
+	TEST_METHOD(SpawnObject_WithOwner_SetsOwner)
+	{
+		auto* Actor = &Spawner.SpawnActor<AActor>();
+		auto& Component = TObjectBuilder<USceneComponent>(Actor).Spawn();
+
+		ASSERT_THAT(AreEqual(Component.GetOwner(), Actor));
+	}
+};
+
+TEST_CLASS(ObjectBuilder_ApiError, "TestFramework.CQTest")
+{
+	FActorTestSpawner Spawner;
+
+	TEST_METHOD(SpawnCalled_AfterSpawn_Actor_Errors)
+	{
+		auto Builder = TObjectBuilder<ATestActorWithProperties>(Spawner);
+		auto& first = Builder.Spawn();
+
+		Assert.ExpectError("");
+		Builder.Spawn();
+	}
+
+	TEST_METHOD(SpawnCalled_AfterSpawn_Object_Errors)
+	{
+		auto Builder = TObjectBuilder<UAnimSequence>();
+		auto& BuiltObject = Builder.Spawn();
+
+		Assert.ExpectError("");
+		Builder.Spawn();
+	}
+
+	TEST_METHOD(SetParam_MissingProperty_Errors)
+	{
+		Assert.ExpectError("Failed to find Object* property");
+		UObject* Object = nullptr;
+
+		TObjectBuilder<AActor>(Spawner)
+			.SetParam("MissingParameter", Object)
+			.Spawn();
+	}
+};
+
+TEST_CLASS(ObjectBuilder_PrimitiveParameterErrors, "TestFramework.CQTest")
+{
+	FActorTestSpawner Spawner;
+
+	TEST_METHOD(SetParam_TypeMismatch_Integers_Errors)
+	{
+		Assert.ExpectError("Type mismatch", 4);
+
+		TObjectBuilder<ATestActorWithProperties>(Spawner)
+			.SetParam("Int8Property", 1.f)
+			.SetParam("Int16Property", 2.f)
+			.SetParam("Int32Property", 3.f)
+			.SetParam("Int64Property", 4.f)
+			.Spawn();
+	}
+
+	TEST_METHOD(SetParam_TypeMismatch_Unsigned_Errors)
+	{
+		Assert.ExpectError("Type mismatch", 4);
+
+		TObjectBuilder<ATestActorWithProperties>(Spawner)
+			.SetParam("UInt8Property", 1.f)
+			.SetParam("UInt16Property", 2.f)
+			.SetParam("UInt32Property", 3.f)
+			.SetParam("UInt64Property", 4.f)
+			.Spawn();
+	}
+
+	TEST_METHOD(SetParam_TypeMismatch_Bool_Errors)
+	{
+		Assert.ExpectError("Type mismatch");
+
+		TObjectBuilder<ATestActorWithProperties>(Spawner)
+			.SetParam("BoolProperty", 42)
+			.Spawn();
+	}
+
+	TEST_METHOD(SetParam_TypeMismatch_Byte_Errors)
+	{
+		Assert.ExpectError("Type mismatch");
+
+		TObjectBuilder<ATestActorWithProperties>(Spawner)
+			.SetParam("ByteProperty", 42)
+			.Spawn();
+	}
+
+	TEST_METHOD(SetParam_TypeMismatch_FloatingPoint_Errors)
+	{
+		Assert.ExpectError("Type mismatch", 2);
+
+		TObjectBuilder<ATestActorWithProperties>(Spawner)
+			.SetParam("FloatProperty", 42)
+			.SetParam("DoubleProperty", 42)
+			.Spawn();
+	}
+
+	TEST_METHOD(SetParam_TypeMismatch_FName_Errors)
+	{
+		Assert.ExpectError("Type mismatch");
+
+		TObjectBuilder<ATestActorWithProperties>(Spawner)
+			.SetParam("NameProperty", 42)
+			.Spawn();
+	}
+
+	TEST_METHOD(SetParam_TypeMismatch_Arrays_Errors)
+	{
+		Assert.ExpectError("Type mismatch", 2);
+
+		TObjectBuilder<ATestActorWithProperties>(Spawner)
+			.SetParam("ArrayProperty", 42)
+			.SetParam("ArrayOfObjectsProperty", 42)
+			.Spawn();
+	}
+
+	TEST_METHOD(SetParam_TypeMismatch_Map_Errors)
+	{
+		Assert.ExpectError("Type mismatch");
+
+		TObjectBuilder<ATestActorWithProperties>(Spawner)
+			.SetParam("MapProperty", 42)
+			.Spawn();
+	}
+
+	TEST_METHOD(SetParam_TypeMismatch_Struct_Errors)
+	{
+		Assert.ExpectError("Type mismatch");
+
+		TObjectBuilder<ATestActorWithProperties>(Spawner)
+			.SetParam("StructProperty", 42)
+			.Spawn();
+	}
+
+	TEST_METHOD(SetParam_TypeMismatch_Vector_Errors)
+	{
+		Assert.ExpectError("Type mismatch");
+
+		TObjectBuilder<ATestActorWithProperties>(Spawner)
+			.SetParam("VectorProperty", 42)
+			.Spawn();
+	}
+
+	TEST_METHOD(SetParam_TypeMismatch_Enums_Errors)
+	{
+		Assert.ExpectError("Type mismatch", 3);
+
+		TObjectBuilder<ATestActorWithProperties>(Spawner)
+			.SetParam("Uint8EnumProperty", 0.0)
+			.SetParam("Int16EnumProperty", 1.f)
+			.SetParam("Uint32EnumProperty", FName("Wrong"))
+			.Spawn();
+	}
+
+	TEST_METHOD(SetParam_TypeMismatch_EnumSize_Errors)
+	{
+		Assert.ExpectError("Type mismatch", 8);
+
+		TObjectBuilder<ATestActorWithProperties>(Spawner)
+			.SetParam("Uint8EnumProperty", ETestUint16::enumone)
+			.SetParam("Uint16EnumProperty", ETestUint8::enumone)
+			.SetParam("Uint32EnumProperty", ETestUint64::enumone)
+			.SetParam("Uint64EnumProperty", ETestUint32::enumone)
+			.SetParam("Int8EnumProperty", ETestInt16::enumone)
+			.SetParam("Int16EnumProperty", ETestInt8::enumone)
+			.SetParam("Int32EnumProperty", ETestInt64::enumone)
+			.SetParam("Int64EnumProperty", ETestInt32::enumone)
+			.Spawn();
+	}
+
+	TEST_METHOD(SetParam_TypeMismatch_TestTObjectPtr_Errors)
+	{
+		Assert.ExpectError("Type mismatch");
+
+		TObjectBuilder<ATestActorWithProperties>(Spawner)
+			.SetParam("TestTObjectPtrProperty", 123)
+			.Spawn();
+	}
+};
+
+TEST_CLASS(ObjectBuilder_StructParameterErrors, "TestFramework.CQTest")
+{
+	FActorTestSpawner Spawner;
+
+	TEST_METHOD(SetParam_ObjectWithNullPointer_Errors)
+	{
+		UObject* NullObject = nullptr;
+		Assert.ExpectError("Type mismatch");
+
+		TObjectBuilder<ATestActorWithProperties>(Spawner)
+			.SetParam("TestTObjectPtrProperty", NullObject)
+			.Spawn();
+	}
+
+	TEST_METHOD(SetParam_UnrelatedObject_Errors)
+	{
+		UActorComponent* UnrelatedStruct = nullptr;
+		Assert.ExpectError("Type mismatch");
+
+		TObjectBuilder<ATestActorWithProperties>(Spawner)
+			.SetParam("TestTObjectPtrProperty", UnrelatedStruct)
+			.Spawn();
+	}
+
+	TEST_METHOD(SetParam_UnrelatedStruct_Errors)
+	{
+		FOtherTestStruct UnrelatedStructType;
+		Assert.ExpectError("Type mismatch");
+
+		auto& BuiltActor = TObjectBuilder<ATestActorWithProperties>(Spawner)
+							   .SetParam("StructProperty", UnrelatedStructType)
+							   .Spawn();
+	}
+
+	TEST_METHOD(SetParam_DerivedStruct_Errors)
+	{
+		FDerivedTestStruct Derived;
+		Assert.ExpectError("Type mismatch");
+
+		auto& BuiltActor = TObjectBuilder<ATestActorWithProperties>(Spawner)
+							   .SetParam("StructProperty", Derived)
+							   .Spawn();
+	}
+};
+
+
+TEST_CLASS(ObjectBuilder_ArrayPropertyErrors, "TestFramework.CQTest")
+{
+	FActorTestSpawner Spawner;
+
+	template <typename T>
+	void Try(T t, bool bIsMap = false)
+	{
+		Assert.ExpectError(bIsMap? "Failed to find" : "Type mismatch", 4);
+
+		TObjectBuilder<ATestActorWithProperties>(Spawner)
+			.SetParam("ArrayProperty", t)
+			.SetParam("ArrayOfObjectsProperty", t)
+			.SetParam("ArrayOfVectorsProperty", t)
+			.SetParam("ArrayOfStructsProperty", t)
+			.Spawn();
+	}
+
+	TEST_METHOD(SetParam_WithPrimitive_Errors)
+	{
+		Try(123);
+	}
+
+	TEST_METHOD(SetParam_WithUnrelatedSimpleArray_Errors)
+	{
+		TArray<float> Unrelated;
+		Try(Unrelated);
+	}
+
+	TEST_METHOD(SetParam_WithUObjectArray_Errors)
+	{
+		TArray<TObjectPtr<UActorComponent>> Unrelated;
+		Try(Unrelated);
+	}
+
+	TEST_METHOD(SetParam_WithNestedSimpleArray_Errors)
+	{
+		TArray<TArray<float>> Unrelated;
+		Try(Unrelated);
+	}
+
+	TEST_METHOD(SetParam_WithNestedUObjectArray_Errors)
+	{
+		TArray<TArray<TObjectPtr<UActorComponent>>> Unrelated;
+		Try(Unrelated);
+	}
+
+	TEST_METHOD(SetParam_WithArrayOfStructs_Errors)
+	{
+		TArray<FOtherTestStruct> Unrelated;
+		Try(Unrelated);
+	}
+
+	TEST_METHOD(SetParam_WithArrayOfMaps_Errors)
+	{
+		TArray<TMap<uint32, uint32>> Unrelated;
+		Try(Unrelated);
+	}
+
+	TEST_METHOD(SetParam_WithArrayOfVectors_Errors)
+	{
+		TArray<TMap<uint32, uint32>> Unrelated;
+		Try(Unrelated);
+	}
+
+	TEST_METHOD(SetParam_WithMap_Errors)
+	{
+		TMap<uint32, uint32> Unrelated;
+		Try(Unrelated, true);
+	}
+
+	TEST_METHOD(SetParam_WithObjectIntMap_Errors)
+	{
+		TMap<UObject*, uint32> Unrelated;
+		Try(Unrelated, true);
+	}
+
+	TEST_METHOD(SetParam_WithIntObjectMap_Errors)
+	{
+		TMap<uint32, UObject*> Unrelated;
+		Try(Unrelated, true);
+	}
+
+	TEST_METHOD(SetParam_WithObjectObjectMap_Errors)
+	{
+		TMap<UObject*, UObject*> Unrelated;
+		Try(Unrelated, true);
+	}
+
+	TEST_METHOD(SetParam_WithIntArrayMap_Errors)
+	{
+		TMap<uint32, TArray<UObject*>> Unrelated;
+		Try(Unrelated, true);
+	}
+
+	TEST_METHOD(SetParam_WithU8Enum_Errors)
+	{
+		ETestUint8 Unrelated = ETestUint8::enumtwo;
+		Try(Unrelated);
+	}
+
+	TEST_METHOD(SetParam_WithU16Enum_Errors)
+	{
+		ETestUint16 Unrelated = ETestUint16::enumone;
+		Try(Unrelated);
+	}
+
+	TEST_METHOD(SetParam_WithU32Enum_Errors)
+	{
+		ETestUint32 Unrelated = ETestUint32::enumtwo;
+		Try(Unrelated);
+	}
+
+	TEST_METHOD(SetParam_WithU64Enum_Errors)
+	{
+		ETestUint64 Unrelated = ETestUint64::enumone;
+		Try(Unrelated);
+	}
+
+	TEST_METHOD(SetParam_WithS8Enum_Errors)
+	{
+		ETestInt8 Unrelated = ETestInt8::enumone;
+		Try(Unrelated);
+	}
+
+	TEST_METHOD(SetParam_WithS16Enum_Errors)
+	{
+		ETestInt16 Unrelated = ETestInt16::enumtwo;
+		Try(Unrelated);
+	}
+
+	TEST_METHOD(SetParam_WithS32Enum_Errors)
+	{
+		ETestInt32 Unrelated = ETestInt32::enumthree;
+		Try(Unrelated);
+	}
+
+	TEST_METHOD(SetParam_WithS64Enum_Errors)
+	{
+		ETestInt64 Unrelated = ETestInt64::enumtwo;
+		Try(Unrelated);
+	}
+};
+
+TEST_CLASS(ObjectBuilder_SetPropertyErrors, "TestFramework.CQTest")
+{
+	FActorTestSpawner Spawner;
+
+	template <typename T>
+	void Try(T t, bool bIsMap = false)
+	{
+		Assert.ExpectError(bIsMap? "Failed to find" : "Type mismatch");
+		TObjectBuilder<ATestActorWithProperties>(Spawner)
+			.SetParam("SetProperty", t)
+			.Spawn();
+	}
+
+		TEST_METHOD(SetParam_WithPrimitive_Errors)
+	{
+		Try(123);
+	}
+
+	TEST_METHOD(SetParam_WithUnrelatedSimpleArray_Errors)
+	{
+		TArray<float> Unrelated;
+		Try(Unrelated);
+	}
+
+	TEST_METHOD(SetParam_WithUObjectArray_Errors)
+	{
+		TArray<UActorComponent*> Unrelated;
+		Try(Unrelated);
+	}
+
+	TEST_METHOD(SetParam_WithNestedSimpleArray_Errors)
+	{
+		TArray<TArray<float>> Unrelated;
+		Try(Unrelated);
+	}
+
+	TEST_METHOD(SetParam_WithNestedUObjectArray_Errors)
+	{
+		TArray<TArray<UActorComponent*>> Unrelated;
+		Try(Unrelated);
+	}
+
+	TEST_METHOD(SetParam_WithArrayOfStructs_Errors)
+	{
+		TArray<FTestStructWithProperties> Unrelated;
+		Try(Unrelated);
+	}
+
+	TEST_METHOD(SetParam_WithArrayOfMaps_Errors)
+	{
+		TArray<TMap<uint32, uint32>> Unrelated;
+		Try(Unrelated);
+	}
+
+	TEST_METHOD(SetParam_WithArrayOfVectors_Errors)
+	{
+		TArray<TMap<uint32, uint32>> Unrelated;
+		Try(Unrelated);
+	}
+
+	TEST_METHOD(SetParam_WithMapWrongKey_Errors)
+	{
+		TMap<UObject*, int32> Unrelated;
+		Try(Unrelated, true);
+	}
+
+	TEST_METHOD(SetParam_WithMapWrongValue_Errors)
+	{
+		TMap<int32, UObject*> Unrelated;
+		Try(Unrelated, true);
+	}
+
+	TEST_METHOD(SetParam_WithMapWrongKeyAndValue_Errors)
+	{
+		TMap<UObject*, UObject*> Unrelated;
+		Try(Unrelated, true);
+	}
+
+	TEST_METHOD(SetParam_WithIntArrayMap_Errors)
+	{
+		TMap<uint32, TArray<UObject*>> Unrelated;
+		Try(Unrelated, true);
+	}
+
+	TEST_METHOD(SetParam_WithU8Enum_Errors)
+	{
+		ETestUint8 Unrelated = ETestUint8::enumtwo;
+		Try(Unrelated);
+	}
+
+	TEST_METHOD(SetParam_WithU16Enum_Errors)
+	{
+		ETestUint16 Unrelated = ETestUint16::enumone;
+		Try(Unrelated);
+	}
+
+	TEST_METHOD(SetParam_WithU32Enum_Errors)
+	{
+		ETestUint32 Unrelated = ETestUint32::enumtwo;
+		Try(Unrelated);
+	}
+
+	TEST_METHOD(SetParam_WithU64Enum_Errors)
+	{
+		ETestUint64 Unrelated = ETestUint64::enumone;
+		Try(Unrelated);
+	}
+
+	TEST_METHOD(SetParam_WithS8Enum_Errors)
+	{
+		ETestInt8 Unrelated = ETestInt8::enumone;
+		Try(Unrelated);
+	}
+
+	TEST_METHOD(SetParam_WithS16Enum_Errors)
+	{
+		ETestInt16 Unrelated = ETestInt16::enumtwo;
+		Try(Unrelated);
+	}
+
+	TEST_METHOD(SetParam_WithS32Enum_Errors)
+	{
+		ETestInt32 Unrelated = ETestInt32::enumthree;
+		Try(Unrelated);
+	}
+
+	TEST_METHOD(SetParam_WithS64Enum_Errors)
+	{
+		ETestInt64 Unrelated = ETestInt64::enumtwo;
+		Try(Unrelated);
+	}
+};
+
+TEST_CLASS(ObjectBuilder_MapPropertyErrors, "TestFramework.CQTest")
+{
+	FActorTestSpawner Spawner;
+
+	template <typename T>
+	void Try(T t)
+	{
+		Assert.ExpectError("Type mismatch");
+
+		TObjectBuilder<ATestActorWithProperties>(Spawner)
+			.SetParam("MapProperty", t)
+			.Spawn();
+	}
+
+	TEST_METHOD(SetParam_WithPrimitive_Errors)
+	{
+		Try(123);
+	}
+
+	TEST_METHOD(SetParam_WithUnrelatedSimpleArray_Errors)
+	{
+		TArray<float> Unrelated;
+		Try(Unrelated);
+	}
+
+	TEST_METHOD(SetParam_WithUObjectArray_Errors)
+	{
+		TArray<UActorComponent*> Unrelated;
+		Try(Unrelated);
+	}
+
+	TEST_METHOD(SetParam_WithNestedSimpleArray_Errors)
+	{
+		TArray<TArray<float>> Unrelated;
+		Try(Unrelated);
+	}
+
+	TEST_METHOD(SetParam_WithNestedUObjectArray_Errors)
+	{
+		TArray<TArray<UActorComponent*>> Unrelated;
+		Try(Unrelated);
+	}
+
+	TEST_METHOD(SetParam_WithArrayOfStructs_Errors)
+	{
+		TArray<FTestStructWithProperties> Unrelated;
+		Try(Unrelated);
+	}
+
+	TEST_METHOD(SetParam_WithArrayOfMaps_Errors)
+	{
+		TArray<TMap<uint32, uint32>> Unrelated;
+		Try(Unrelated);
+	}
+
+	TEST_METHOD(SetParam_WithArrayOfVectors_Errors)
+	{
+		TArray<TMap<uint32, uint32>> Unrelated;
+		Try(Unrelated);
+	}
+
+	TEST_METHOD(SetParam_WithMapWrongKey_Errors)
+	{
+		TMap<UObject*, int32> Unrelated;
+		Try(Unrelated);
+	}
+
+	TEST_METHOD(SetParam_WithMapWrongValue_Errors)
+	{
+		TMap<int32, UObject*> Unrelated;
+		Try(Unrelated);
+	}
+
+	TEST_METHOD(SetParam_WithMapWrongKeyAndValue_Errors)
+	{
+		TMap<UObject*, UObject*> Unrelated;
+		Try(Unrelated);
+	}
+
+	TEST_METHOD(SetParam_WithIntArrayMap_Errors)
+	{
+		TMap<uint32, TArray<UObject*>> Unrelated;
+		Try(Unrelated);
+	}
+
+	TEST_METHOD(SetParam_WithU8Enum_Errors)
+	{
+		ETestUint8 Unrelated = ETestUint8::enumtwo;
+		Try(Unrelated);
+	}
+
+	TEST_METHOD(SetParam_WithU16Enum_Errors)
+	{
+		ETestUint16 Unrelated = ETestUint16::enumone;
+		Try(Unrelated);
+	}
+
+	TEST_METHOD(SetParam_WithU32Enum_Errors)
+	{
+		ETestUint32 Unrelated = ETestUint32::enumtwo;
+		Try(Unrelated);
+	}
+
+	TEST_METHOD(SetParam_WithU64Enum_Errors)
+	{
+		ETestUint64 Unrelated = ETestUint64::enumone;
+		Try(Unrelated);
+	}
+
+	TEST_METHOD(SetParam_WithS8Enum_Errors)
+	{
+		ETestInt8 Unrelated = ETestInt8::enumone;
+		Try(Unrelated);
+	}
+
+	TEST_METHOD(SetParam_WithS16Enum_Errors)
+	{
+		ETestInt16 Unrelated = ETestInt16::enumtwo;
+		Try(Unrelated);
+	}
+
+	TEST_METHOD(SetParam_WithS32Enum_Errors)
+	{
+		ETestInt32 Unrelated = ETestInt32::enumthree;
+		Try(Unrelated);
+	}
+
+	TEST_METHOD(SetParam_WithS64Enum_Errors)
+	{
+		ETestInt64 Unrelated = ETestInt64::enumtwo;
+		Try(Unrelated);
+	}
+};
+
+TEST_CLASS(ObjectBuilder_PostSpawnErrors, "TestFramework.CQTest")
+{
+	FActorTestSpawner Spawner;
+
+	template <typename T>
+	void Try(const FName& Name, T t)
+	{
+		TObjectBuilder<ATestActorWithProperties> Builder(Spawner);
+		Builder.Spawn();
+
+		Assert.ExpectError(TEXT("Tried to SetParam"));
+
+		Builder.SetParam(Name, t);
+	}
+
+	TEST_METHOD(SetIntProperty_PostSpawn_Errors)
+	{
+		Try("IntProperty", 42);
+	}
+
+	TEST_METHOD(SetMapProperty_PostSpawn_Errors)
+	{
+		TMap<int32, int32> Map;
+		Try("MapProperty", Map);
+	}
+
+	TEST_METHOD(SetArrayProperty_PostSpawn_Errors)
+	{
+		TArray<int32> Array;
+		Try("ArrayProperty", Array);
+	}
+
+	TEST_METHOD(SetObjectProperty_PostSpawn_Errors)
+	{
+		UObject* Object = nullptr;
+		Try("UObjectProperty", Object);
+	}
+
+	TEST_METHOD(SetStructProperty_PostSpawn_Errors)
+	{
+		FTestStructWithProperties Struct;
+		Try("UStructProperty", Struct);
+	}
+
+	TEST_METHOD(SetEnumPropery_PostSpawn_Errors)
+	{
+		TObjectBuilder<ATestActorWithProperties> Builder(Spawner);
+		Builder.Spawn();
+
+		Assert.ExpectError(TEXT("Tried to SetParam"), 8);
+
+		Builder.SetParam("Uint8EnumProperty", ETestUint8::enumone);
+		Builder.SetParam("Uint16EnumProperty", ETestUint16::enumone);
+		Builder.SetParam("Uint32EnumProperty", ETestUint32::enumone);
+		Builder.SetParam("Uint64EnumProperty", ETestUint64::enumone);
+
+		Builder.SetParam("Int8EnumProperty", ETestInt8::enumone);
+		Builder.SetParam("Int16EnumProperty", ETestInt16::enumone);
+		Builder.SetParam("Int32EnumProperty", ETestInt32::enumone);
+		Builder.SetParam("Int64EnumProperty", ETestInt64::enumone);
+	}
+
+	TEST_METHOD(AddComponentTo_PostSpawn_Errors)
+	{
+		TObjectBuilder<ATestActorWithProperties> Builder(Spawner);
+		Builder.Spawn();
+
+		Assert.ExpectError(TEXT("Tried to AddComponentTo"));
+		Builder.AddComponentTo<UActorComponent>();
+	}
+
+	TEST_METHOD(AddChildActorComponentTo_PostSpawn_Errors)
+	{
+		TObjectBuilder<ATestActorWithProperties> Builder(Spawner);
+		Builder.Spawn();
+
+		Assert.ExpectError(TEXT("Tried to AddChildActorComponentTo"));
+		Builder.AddChildActorComponentTo<AActor>();
+	}
+};
\ No newline at end of file
diff --git a/Engine/Plugins/Tests/CQTest/Source/CQTestTests/Private/SpawnHelperTests.cpp b/Engine/Plugins/Tests/CQTest/Source/CQTestTests/Private/SpawnHelperTests.cpp
new file mode 100644
--- /dev/null
+++ b/Engine/Plugins/Tests/CQTest/Source/CQTestTests/Private/SpawnHelperTests.cpp
@@ -0,0 +1,30 @@
+// Copyright Epic Games, Inc. All Rights Reserved.
+
+#include "CQTest.h"
+#include "ActorTestSpawner.h"
+
+TEST_CLASS(SpawnHelperTests, "TestFramework.CQTest")
+{
+	FActorTestSpawner Spawner;
+
+	TEST_METHOD(SpawnActor_WithNameParameter_ReturnsActorWithGivenName)
+	{
+		const FName ExpectedActorName = TEXT("AnyActorName");
+		FActorSpawnParameters SpawnParameters;
+		SpawnParameters.Name = ExpectedActorName;
+
+		AActor& Actor = Spawner.SpawnActor<AActor>(SpawnParameters);
+
+		ASSERT_THAT(AreEqual( Actor.GetFName(), ExpectedActorName ));
+	}
+	
+	TEST_METHOD(AddComponentTo_SpawnedActor_SetsActorAsComponentOwner)
+	{
+		AActor& Owner = Spawner.SpawnActor<AActor>();
+		auto* TestComponent = NewObject<USceneComponent>(&Owner);
+		TestComponent->RegisterComponent();
+
+		ASSERT_THAT(IsNotNull(TestComponent));
+		ASSERT_THAT(AreEqual(TestComponent->GetOwner(), &Owner ));
+	}
+};
\ No newline at end of file
diff --git a/Engine/Plugins/Tests/CQTest/Source/CQTestTests/Private/TestActorWithProperties.h b/Engine/Plugins/Tests/CQTest/Source/CQTestTests/Private/TestActorWithProperties.h
new file mode 100644
--- /dev/null
+++ b/Engine/Plugins/Tests/CQTest/Source/CQTestTests/Private/TestActorWithProperties.h
@@ -0,0 +1,208 @@
+// Copyright Epic Games, Inc. All Rights Reserved.
+
+#pragma once
+
+// -----------------------------------------------------------------------------
+
+#include "CoreMinimal.h"
+#include "GameFramework/Actor.h"
+
+#include "TestActorWithProperties.generated.h"
+
+// -----------------------------------------------------------------------------
+
+USTRUCT()
+struct FTestStructWithProperties
+{
+	GENERATED_BODY()
+public:
+	UPROPERTY()
+	int32 StructInt32Property = 0;
+
+	UPROPERTY()
+	TArray<int32> StructArrayProperty;
+};
+
+USTRUCT()
+struct FDerivedTestStruct : public FTestStructWithProperties
+{
+	GENERATED_BODY()
+public:
+	int64 StructInt64Property = 0;
+};
+
+USTRUCT()
+struct FOtherTestStruct
+{
+	GENERATED_BODY()
+public:
+	UPROPERTY()
+	float StructFloatProperty = 0.0f;
+
+	UPROPERTY()
+	int64 StructInt64Property = 0;
+
+	UPROPERTY()
+	TMap<int32, int32> StructMapProperty;
+};
+
+UENUM()
+enum class ETestUint8 : uint8
+{
+	enumone,
+	enumtwo,
+	enumthree
+};
+UENUM()
+enum class ETestInt8 : int8
+{
+	enumone,
+	enumtwo,
+	enumthree
+};
+UENUM()
+enum class ETestInt16 : int16
+{
+	enumone,
+	enumtwo,
+	enumthree
+};
+UENUM()
+enum class ETestUint16 : uint16
+{
+	enumone,
+	enumtwo,
+	enumthree
+};
+UENUM()
+enum class ETestInt32 : int32
+{
+	enumone,
+	enumtwo,
+	enumthree
+};
+UENUM()
+enum class ETestUint32 : uint32
+{
+	enumone,
+	enumtwo,
+	enumthree
+};
+UENUM()
+enum class ETestInt64 : int64
+{
+	enumone,
+	enumtwo,
+	enumthree
+};
+UENUM()
+enum class ETestUint64 : uint64
+{
+	enumone,
+	enumtwo,
+	enumthree
+};
+
+UCLASS()
+class ATestActorWithProperties : public AActor
+{
+	GENERATED_BODY()
+
+public:
+	UPROPERTY()
+	bool BoolProperty;
+
+	UPROPERTY()
+	uint8 ByteProperty;
+
+	UPROPERTY()
+	int8 Int8Property;
+
+	UPROPERTY()
+	uint8 UInt8Property;
+
+	UPROPERTY()
+	uint16 UInt16Property;
+
+	UPROPERTY()
+	int16 Int16Property;
+
+	UPROPERTY()
+	uint32 UInt32Property;
+
+	UPROPERTY()
+	int64 Int64Property;
+
+	UPROPERTY()
+	uint64 UInt64Property;
+
+	UPROPERTY()
+	FTestStructWithProperties StructProperty;
+
+	UPROPERTY()
+	float FloatProperty;
+
+	UPROPERTY()
+	double DoubleProperty;
+
+	UPROPERTY()
+	int32 Int32Property;
+
+	UPROPERTY()
+	TObjectPtr<class UAnimSequence> TestTObjectPtrProperty;
+
+	UPROPERTY()
+	TScriptInterface<UObject> InterfaceProperty;
+
+	UPROPERTY()
+	FName NameProperty;
+
+	UPROPERTY()
+	TArray<int32> ArrayProperty;
+
+	UPROPERTY()
+	TMap<int32, int32> MapProperty;
+
+	UPROPERTY()
+	TSet<int32> SetProperty;
+
+	UPROPERTY()
+	FVector VectorProperty;
+
+	UPROPERTY()
+	ETestUint8 Uint8EnumProperty;
+
+	UPROPERTY()
+	ETestInt8 Int8EnumProperty;
+
+	UPROPERTY()
+	ETestInt16 Int16EnumProperty;
+
+	UPROPERTY()
+	ETestUint16 Uint16EnumProperty;
+
+	UPROPERTY()
+	ETestInt32 Int32EnumProperty;
+
+	UPROPERTY()
+	ETestUint32 Uint32EnumProperty;
+
+	UPROPERTY()
+	ETestInt64 Int64EnumProperty;
+
+	UPROPERTY()
+	ETestUint64 Uint64EnumProperty;
+
+	UPROPERTY()
+	TArray<TObjectPtr<ATestActorWithProperties>> ArrayOfObjectsProperty;
+
+	UPROPERTY()
+	TArray<FVector> ArrayOfVectorsProperty;
+
+	UPROPERTY()
+	TArray<FTestStructWithProperties> ArrayOfStructsProperty;
+};
+
+// -----------------------------------------------------------------------------
+// -----------------------------------------------------------------------------
+// -----------------------------------------------------------------------------
diff --git a/Engine/Plugins/Tests/CQTest/Source/CQTestTests/Private/TestCommandsTests.cpp b/Engine/Plugins/Tests/CQTest/Source/CQTestTests/Private/TestCommandsTests.cpp
new file mode 100644
--- /dev/null
+++ b/Engine/Plugins/Tests/CQTest/Source/CQTestTests/Private/TestCommandsTests.cpp
@@ -0,0 +1,133 @@
+// Copyright Epic Games, Inc. All Rights Reserved.
+
+#include "CQTest.h"
+
+struct FCommandLog
+{
+	TArray<FString> Commands;
+};
+
+class FNamedCommand : public IAutomationLatentCommand
+{
+public:
+	FNamedCommand(TArray<FString>& CommandLog, FString Name)
+		: Log(CommandLog), CommandName(Name) {}
+
+	bool Update() override
+	{
+		Log.Add(CommandName);
+		return true;
+	}
+
+	TArray<FString>& Log;
+	FString CommandName;
+};
+
+class FTickingNamedCommand : public FNamedCommand
+{
+public:
+	FTickingNamedCommand(TArray<FString>& CommandLog, FString Name, int32 Ticks)
+		: FNamedCommand(CommandLog, Name), ExpectedCount(Ticks) {}
+		
+	bool Update() override
+	{
+		if (CurrentCount == ExpectedCount)
+		{
+			return true;
+		}
+
+		Log.Add(CommandName);
+		CurrentCount++;
+		return false;
+	}
+
+	int32 ExpectedCount{ 0 };
+	int32 CurrentCount{ 0 };
+};
+
+TEST_CLASS(RunSequenceTests, "TestFramework.CQTest")
+{
+	const TArray<FString> Names = 
+	{ 
+		"Zero",
+		"One",
+		"Two",
+		"Three",
+		"Four",
+	};
+
+	TFunction<bool(RunSequenceTests*)> Assertion;
+	TArray<FString> CommandLog;
+
+	AFTER_EACH()
+	{
+		ASSERT_THAT(IsTrue(Assertion(this)));
+	}
+
+	TEST_METHOD(RunSequence_WithZeroCommands_DoesNotFail)
+	{
+		AddCommand(new FRunSequence());
+		Assertion = [](RunSequenceTests* test) {
+			return test->CommandLog.IsEmpty();
+		};
+	}
+
+	TEST_METHOD(RunSequence_WithOneCommand_RunsCommand)
+	{
+		AddCommand(new FRunSequence(MakeShared<FNamedCommand>(CommandLog, Names[0])));
+		Assertion = [](RunSequenceTests* test) {
+			return test->CommandLog.Num() == 1 && test->CommandLog[0] == test->Names[0];
+		};
+	}
+
+	TEST_METHOD(RunSequence_WithNamedCommands_RunsCommandsInOrder)
+	{
+		TArray<TSharedPtr<FNamedCommand>> Commands;
+		Commands.Add(MakeShared<FNamedCommand>(CommandLog, Names[0]));
+		Commands.Add(MakeShared<FNamedCommand>(CommandLog, Names[1]));
+		Commands.Add(MakeShared<FNamedCommand>(CommandLog, Names[2]));
+		Commands.Add(MakeShared<FNamedCommand>(CommandLog, Names[3]));
+		Commands.Add(MakeShared<FNamedCommand>(CommandLog, Names[4]));
+		
+		AddCommand(new FRunSequence(Commands));
+
+		Assertion = [](RunSequenceTests* test) {
+			return test->CommandLog == test->Names;
+		};
+	}
+
+	TEST_METHOD(RunSequence_WithTickingCommands_RunsCommandsInOrder)
+	{
+		TArray<TSharedPtr<FTickingNamedCommand>> Commands;
+		Commands.Add(MakeShared<FTickingNamedCommand>(CommandLog, Names[0], 3));
+		Commands.Add(MakeShared<FTickingNamedCommand>(CommandLog, Names[1], 3));
+		Commands.Add(MakeShared<FTickingNamedCommand>(CommandLog, Names[2], 3));
+		Commands.Add(MakeShared<FTickingNamedCommand>(CommandLog, Names[3], 3));
+		Commands.Add(MakeShared<FTickingNamedCommand>(CommandLog, Names[4], 3));
+
+		AddCommand(new FRunSequence(Commands));
+
+		Assertion = [](RunSequenceTests* test) {
+			if (test->CommandLog.Num() != 15)
+			{
+				return false;
+			}
+
+			int32 NameIndex = -1;
+			for (int32 CommandIndex = 0; CommandIndex < test->CommandLog.Num(); CommandIndex++)
+			{
+				if (CommandIndex % 3 == 0)
+				{
+					NameIndex++;
+				}
+
+				if (test->CommandLog[CommandIndex] != test->Names[NameIndex])
+				{
+					return false;
+				}
+			}
+
+			return true;
+		};
+	}
+};
\ No newline at end of file
