Fix unity build in CQTest MapTestSpawner
Add FPaths include

#issue 281808

diff --git a/Engine/Plugins/Tests/CQTest/Source/CQTest/Private/MapTestSpawner.cpp b/Engine/Plugins/Tests/CQTest/Source/CQTest/Private/MapTestSpawner.cpp
--- a/Engine/Plugins/Tests/CQTest/Source/CQTest/Private/MapTestSpawner.cpp
+++ b/Engine/Plugins/Tests/CQTest/Source/CQTest/Private/MapTestSpawner.cpp
@@ -6,6 +6,7 @@
 #include "Tests/AutomationCommon.h"
 #include "GameFramework/PlayerController.h"
 #include "Engine/Engine.h"
+#include "Misc/Paths.h"
 
 void FMapTestSpawner::AddWaitUntilLoadedCommand(FAutomationTestBase* TestRunner)
 {
