Temp Fix for Hololens Build break with CQTestConverter
Disable the plugin on that platform

#Issue 281796

diff --git a/Engine/Platforms/Hololens/Plugins/Tests/CQTest/CQTest_HoloLens.uplugin b/Engine/Platforms/Hololens/Plugins/Tests/CQTest/CQTest_HoloLens.uplugin
new file mode 100644
--- /dev/null
+++ b/Engine/Platforms/Hololens/Plugins/Tests/CQTest/CQTest_HoloLens.uplugin
@@ -0,0 +1,20 @@
+{
+  "FileVersion": 3,
+  "bIsPluginExtension": true,
+  "Modules": [
+    {
+      "Name": "CQTest",
+      "Type": "DeveloperTool",
+      "PlatformDenyList": [
+        "HoloLens"
+      ]
+    },
+    {
+      "Name": "CQTestTests",
+      "Type": "DeveloperTool",
+      "PlatformDenyList": [
+        "HoloLens"
+      ]
+    }
+  ]
+}
\ No newline at end of file
diff --git a/Engine/Platforms/Hololens/Plugins/Tests/CQTest/Source/CQTest/CQTest_HoloLens.Build.cs b/Engine/Platforms/Hololens/Plugins/Tests/CQTest/Source/CQTest/CQTest_HoloLens.Build.cs
new file mode 100644
--- /dev/null
+++ b/Engine/Platforms/Hololens/Plugins/Tests/CQTest/Source/CQTest/CQTest_HoloLens.Build.cs
@@ -0,0 +1,12 @@
+// Copyright Epic Games, Inc. All Rights Reserved.
+
+namespace UnrealBuildTool.Rules
+{
+	public class CQTest_HoloLens : CQTest
+	{
+		public CQTest_HoloLens(ReadOnlyTargetRules Target) : base(Target)
+		{
+			throw new System.Exception("Should not be here - HoloLens in the uplugin's PlatformDenyList");
+		}
+	}
+}
diff --git a/Engine/Platforms/Hololens/Plugins/Tests/CQTest/Source/CQTestTests/CQTestTests_HoloLens.Build.cs b/Engine/Platforms/Hololens/Plugins/Tests/CQTest/Source/CQTestTests/CQTestTests_HoloLens.Build.cs
new file mode 100644
--- /dev/null
+++ b/Engine/Platforms/Hololens/Plugins/Tests/CQTest/Source/CQTestTests/CQTestTests_HoloLens.Build.cs
@@ -0,0 +1,12 @@
+// Copyright Epic Games, Inc. All Rights Reserved.
+
+namespace UnrealBuildTool.Rules
+{
+	public class CQTestTests_HoloLens : CQTestTests
+	{
+		public CQTestTests_HoloLens(ReadOnlyTargetRules Target) : base(Target)
+		{
+			throw new System.Exception("Should not be here - HoloLens in the uplugin's PlatformDenyList");
+		}
+	}
+}
