Fix include cycle

diff --git a/Engine/Plugins/Tests/CQTest/Source/CQTest/Public/Assert/NoDiscardAsserter.inl b/Engine/Plugins/Tests/CQTest/Source/CQTest/Public/Assert/NoDiscardAsserter.inl
--- a/Engine/Plugins/Tests/CQTest/Source/CQTest/Public/Assert/NoDiscardAsserter.inl
+++ b/Engine/Plugins/Tests/CQTest/Source/CQTest/Public/Assert/NoDiscardAsserter.inl
@@ -1,6 +1,5 @@
 // Copyright Epic Games, Inc. All Rights Reserved.
 
-#include "Assert/NoDiscardAsserter.h"
 #include "Internationalization/Regex.h"
 #include "Assert/CQTestConvert.h"
 
