CQTest - exclude some tests when targetting shipping

#issue 281808
#rnx

diff --git a/Engine/Plugins/Tests/CQTest/CQTest.uplugin b/Engine/Plugins/Tests/CQTest/CQTest.uplugin
--- a/Engine/Plugins/Tests/CQTest/CQTest.uplugin
+++ b/Engine/Plugins/Tests/CQTest/CQTest.uplugin
@@ -4,8 +4,8 @@
   "VersionName": "1.0",
   "FriendlyName": "Code Quality Unreal Test Plugin",
   "Description": "Simplified testing classes for Unreal Engine",
-  "Category": "Code Quality",
-  "CreatedBy": "Studios Quality - Code Quality",
+  "Category": "Testing",
+  "CreatedBy": "Code Quality",
   "CreatedByURL": "",
   "DocsURL": "",
   "MarketplaceURL": "",
diff --git a/Engine/Plugins/Tests/CQTest/Source/CQTest/Private/MapTestSpawner.cpp b/Engine/Plugins/Tests/CQTest/Source/CQTest/Private/MapTestSpawner.cpp
--- a/Engine/Plugins/Tests/CQTest/Source/CQTest/Private/MapTestSpawner.cpp
+++ b/Engine/Plugins/Tests/CQTest/Source/CQTest/Private/MapTestSpawner.cpp
@@ -10,6 +10,7 @@
 
 void FMapTestSpawner::AddWaitUntilLoadedCommand(FAutomationTestBase* TestRunner)
 {
+#if WITH_AUTOMATION_TESTS
 	check(PieWorld == nullptr);
 
 	const FString FileName = FString::Printf(TEXT("%s.%s"), *MapName, *MapName);
@@ -29,6 +30,9 @@ void FMapTestSpawner::AddWaitUntilLoadedCommand(FAutomationTestBase* TestRunner)
 
 		return false;
 	}));
+#else
+	checkf(false, TEXT("AddWaitUntilLoadedCommand can't call AutomationOpenMap if WITH_AUTOMATION_TESTS=false"));
+#endif
 }
 
 UWorld* FMapTestSpawner::CreateWorld()
diff --git a/Engine/Plugins/Tests/CQTest/Source/CQTest/Public/Assert/CQTestConvert.h b/Engine/Plugins/Tests/CQTest/Source/CQTest/Public/Assert/CQTestConvert.h
--- a/Engine/Plugins/Tests/CQTest/Source/CQTest/Public/Assert/CQTestConvert.h
+++ b/Engine/Plugins/Tests/CQTest/Source/CQTest/Public/Assert/CQTestConvert.h
@@ -55,7 +55,7 @@ namespace
 namespace CQTestConvert
 {
 	template <typename T>
-	FString ToString(const T& Input)
+	inline FString ToString(const T& Input)
 	{
 		if constexpr (THasToString<T>::value)
 		{
@@ -79,13 +79,13 @@ namespace CQTestConvert
 	}
 
 	template <>
-	FString ToString(const FString& Input)
+	inline FString ToString(const FString& Input)
 	{
 		return Input;
 	}
 
 	template<typename TElement, typename TAllocator>
-	FString ToString(const TArray<TElement, TAllocator>& Input)
+	inline FString ToString(const TArray<TElement, TAllocator>& Input)
 	{
 		FString result = FString::JoinBy(Input, TEXT(", "), [](const TElement& e) {
 			return CQTestConvert::ToString(e);
@@ -94,7 +94,7 @@ namespace CQTestConvert
 	}
 
 	template <typename TElement, typename TAllocator>
-	FString ToString(const TSet<TElement, TAllocator>& Input)
+	inline FString ToString(const TSet<TElement, TAllocator>& Input)
 	{
 		FString result = FString::JoinBy(Input, TEXT(", "), [](const TElement& e) {
 			return CQTestConvert::ToString(e);
@@ -103,7 +103,7 @@ namespace CQTestConvert
 	}
 
 	template<typename TKey, typename TValue>
-	FString ToString(const TMap<TKey, TValue>& Input)
+	inline FString ToString(const TMap<TKey, TValue>& Input)
 	{
 		return FString::JoinBy(Input, TEXT(", "), [](const auto& kvp) {
 			return TEXT("{") + CQTestConvert::ToString(kvp.Key) + TEXT(":") + CQTestConvert::ToString(kvp.Value) + TEXT("}");
diff --git a/Engine/Plugins/Tests/CQTest/Source/CQTest/Public/Assert/NoDiscardAsserter.inl b/Engine/Plugins/Tests/CQTest/Source/CQTest/Public/Assert/NoDiscardAsserter.inl
--- a/Engine/Plugins/Tests/CQTest/Source/CQTest/Public/Assert/NoDiscardAsserter.inl
+++ b/Engine/Plugins/Tests/CQTest/Source/CQTest/Public/Assert/NoDiscardAsserter.inl
@@ -3,7 +3,7 @@
 #include "Internationalization/Regex.h"
 #include "Assert/CQTestConvert.h"
 
-void FNoDiscardAsserter::ExpectError(FString Error, int32 Count)
+inline void FNoDiscardAsserter::ExpectError(FString Error, int32 Count)
 {
 	static const FRegexPattern Pattern(TEXT("([\\[\\]{}()^$.?\\\\*+|])"));
 	FRegexMatcher Matcher(Pattern, Error);
diff --git a/Engine/Plugins/Tests/CQTest/Source/CQTestTests/Private/CQTestConvertTests.cpp b/Engine/Plugins/Tests/CQTest/Source/CQTestTests/Private/CQTestConvertTests.cpp
--- a/Engine/Plugins/Tests/CQTest/Source/CQTestTests/Private/CQTestConvertTests.cpp
+++ b/Engine/Plugins/Tests/CQTest/Source/CQTestTests/Private/CQTestConvertTests.cpp
@@ -11,7 +11,7 @@ struct TypeWithExplicitOverride
 };
 
 template <>
-FString CQTestConvert::ToString(const TypeWithExplicitOverride&)
+inline FString CQTestConvert::ToString(const TypeWithExplicitOverride&)
 {
 	return OverrideStr;
 }
@@ -41,7 +41,7 @@ struct TypeWithToStringAndOverride
 };
 
 template<>
-FString CQTestConvert::ToString(const TypeWithToStringAndOverride&) {
+inline FString CQTestConvert::ToString(const TypeWithToStringAndOverride&) {
 	return OverrideStr;
 }
 
diff --git a/Engine/Plugins/Tests/CQTest/Source/CQTestTests/Private/ObjectBuilderTests.cpp b/Engine/Plugins/Tests/CQTest/Source/CQTestTests/Private/ObjectBuilderTests.cpp
--- a/Engine/Plugins/Tests/CQTest/Source/CQTestTests/Private/ObjectBuilderTests.cpp
+++ b/Engine/Plugins/Tests/CQTest/Source/CQTestTests/Private/ObjectBuilderTests.cpp
@@ -9,6 +9,8 @@
 #include "ActorTestSpawner.h"
 #include "CQTest.h"
 
+#if WITH_AUTOMATION_WORKER
+
 TEST_CLASS(ObjectBuilder_Success, "TestFramework.CQTest")
 {
 	FActorTestSpawner Spawner;
@@ -1154,4 +1156,5 @@ TEST_CLASS(ObjectBuilder_PostSpawnErrors, "TestFramework.CQTest")
 		Assert.ExpectError(TEXT("Tried to AddChildActorComponentTo"));
 		Builder.AddChildActorComponentTo<AActor>();
 	}
-};
\ No newline at end of file
+};
+#endif
\ No newline at end of file
