CQTest - fix issue with some platform compilation

#rnx

diff --git a/Engine/Platforms/Hololens/Plugins/Tests/CQTest/CQTest_HoloLens.uplugin b/Engine/Platforms/Hololens/Plugins/Tests/CQTest/CQTest_HoloLens.uplugin
deleted file mode 100644
diff --git a/Engine/Platforms/Hololens/Plugins/Tests/CQTest/Source/CQTest/CQTest_HoloLens.Build.cs b/Engine/Platforms/Hololens/Plugins/Tests/CQTest/Source/CQTest/CQTest_HoloLens.Build.cs
deleted file mode 100644
diff --git a/Engine/Platforms/Hololens/Plugins/Tests/CQTest/Source/CQTestTests/CQTestTests_HoloLens.Build.cs b/Engine/Platforms/Hololens/Plugins/Tests/CQTest/Source/CQTestTests/CQTestTests_HoloLens.Build.cs
deleted file mode 100644
diff --git a/Engine/Plugins/Tests/CQTest/Source/CQTest/Public/Assert/CQTestConvert.h b/Engine/Plugins/Tests/CQTest/Source/CQTest/Public/Assert/CQTestConvert.h
--- a/Engine/Plugins/Tests/CQTest/Source/CQTest/Public/Assert/CQTestConvert.h
+++ b/Engine/Plugins/Tests/CQTest/Source/CQTest/Public/Assert/CQTestConvert.h
@@ -13,16 +13,15 @@ namespace CQTestConvert
 
 namespace
 {
-	template <typename T, typename... Args>
-	class THasToString
+	template <typename T, typename = void>
+	struct THasToString : std::false_type
 	{
-		template <typename C, typename = decltype(std::declval<C>().ToString(std::declval<Args>()...))>
-		static std::true_type test(int);
-		template<typename C>
-		static std::false_type test(...);
+	};
 
-	public:
-		static constexpr bool value = decltype(test<T>(0))::value;
+	template<typename T>
+	struct THasToString<T, std::void_t<decltype(std::declval<T>().ToString())>>
+		: std::is_same<FString, decltype(std::declval<T>().ToString())>::type
+	{
 	};
 
 	template<typename T>
@@ -39,10 +38,34 @@ namespace
 		static constexpr bool value = decltype(test<T>(0))::value;
 	};
 
-	static_assert(THasToString<FName>::value, "FName should have ToString");
+	static_assert(std::is_same<FString, decltype(std::declval<FName>().ToString())>::value);
 	static_assert(THasOStream<int>::value, "int should have an OStream operator");
 	static_assert(THasOStream<int32>::value, "int32 should have an OStream operator");
 
+	struct StructWithToString
+	{
+		FString ToString() {
+			return FString();
+		}
+	};
+
+	struct StructWithConstToString
+	{
+		FString ToString() const {
+			return FString();
+		}
+	};
+
+	static_assert(THasToString<StructWithToString>::value, "Struct with ToString should have ToString");
+	static_assert(THasToString<StructWithConstToString>::value, "Struct with ToString const should have ToString");
+
+	struct StructWithToStringWrongReturnType
+	{
+		int ToString() { return 42; }
+	};
+
+	static_assert(!THasToString<StructWithToStringWrongReturnType>::value, "Struct with wrong return type on ToString should not have ToString");
+
 	struct SomeTestStruct
 	{
 	};
