CQTest - Add a check to CreatePhysicsScene returned value in  FActorTestSpawner::CreateWorld()

#rnx

diff --git a/Engine/Plugins/Tests/CQTest/Source/CQTest/Private/ActorTestSpawner.cpp b/Engine/Plugins/Tests/CQTest/Source/CQTest/Private/ActorTestSpawner.cpp
--- a/Engine/Plugins/Tests/CQTest/Source/CQTest/Private/ActorTestSpawner.cpp
+++ b/Engine/Plugins/Tests/CQTest/Source/CQTest/Private/ActorTestSpawner.cpp
@@ -85,7 +85,7 @@ UWorld* FActorTestSpawner::CreateWorld()
 	WorldContext.SetCurrentWorld(Result);
 
 	Result->InitializeActorsForPlay(FURL());
-	Result->CreatePhysicsScene();
+	check(Result->GetPhysicsScene() != nullptr);
 
 	return Result;
 }
