PR #10379: First pass at network tests

diff --git a/Engine/Plugins/Tests/CQTest/README.md b/Engine/Plugins/Tests/CQTest/README.md
--- a/Engine/Plugins/Tests/CQTest/README.md
+++ b/Engine/Plugins/Tests/CQTest/README.md
@@ -185,20 +185,20 @@ Constructors (and destructors) are available.  Destructors shouldn't throw, and
 ```cpp
 	TEST_CLASS(SomeTestClass, "Game.Test")
 	{
-		bool constructed = false;
+		bool bConstructed = false;
 		SomeTestClass()
-			: constructed(true)
+			: bConstructed(true)
 		{
 		}
 		
 		TEST_METHOD(ConstructorIsCalled)
 		{
-			ASSERT_THAT(IsTrue(constructed));
+			ASSERT_THAT(IsTrue(bConstructed));
 		}
 	};
 ```	
 
-Latent actions are supported with the TEST_CLASS macro.  Each step will complete all latent actions before moving to the next
+Latent actions are supported with the TEST_CLASS macro.  Each step will complete all latent actions before moving to the next.  If an assertion is raised during a latent action, then no further latent actions will be processed.  The AFTER_EACH method will still be invoked though.
 
 ```cpp
     TEST_CLASS(LatentActionTest, "Game.Test") 
@@ -238,10 +238,20 @@ Also available for commands is a fluent command builder
 
 The framework will ensure that all of those commands happen in order using a future pattern.
 Similarly, the framework will ensure that a test can await a ticking object.  See GameObjectsTickTest for an example
+One word of caution, the framework does not currently support adding latent actions from within latent actions.
+Instead, it is better to add the actions as a series of self-contained steps.
+ 
 
 # Extending the framework
 The framework has been designed to allow for extensions in a couple areas.  See ExtensionTests.cpp for in-code examples.
 
+## Test Components
+This testing framework embraces composition over inheritence.  Creating new components should be the default mechanism for extending the framework.  Some of the components available to you are:
+  ActorTestSpawner - Allows a test to spawn actors, and manages their despawning.
+  MapTestSpawner - Creates a map and opens a level.  Allows tests to spawn actors in that world.
+  BlueprintHelper - Eases the ability for a test to spawn Blueprint objects, intended to be used with MapTestSpawner.
+  PIENetworkComponent - Allows tests to create a server and a collection of clients.  Good for testing replication.
+
 ## Assertions
 Not all platforms support exceptions, and so the assertions are unable to rely on them.
 There are a few options here:
@@ -291,15 +301,16 @@ See CQTestTests/Private/ExtensionTests.cpp for an example
 Below is some untested example code to inspire ideas
 
 ```cpp
-	struct FluentAsserter : public ThrowingAsserter
+	struct FluentAsserter
 	{
 	private:
 		int CurrentIntValue = 0;
 		TArray<FString> Errors;
+		FAutomationTestBase& TestRunner;
 		
 	public:
-		FluentAsserter(FAutomationTestBase& testRunner)
-			: ThrowingAsserter(testRunner)
+		FluentAsserter(FAutomationTestBase& InTestRunner)
+			: TestRunner(InTestRunner)
 		{
 		}
 		
diff --git a/Engine/Plugins/Tests/CQTest/Source/CQTest/Private/ActorTestSpawner.cpp b/Engine/Plugins/Tests/CQTest/Source/CQTest/Private/ActorTestSpawner.cpp
deleted file mode 100644
diff --git a/Engine/Plugins/Tests/CQTest/Source/CQTest/Private/CQTestBlueprintHelper.cpp b/Engine/Plugins/Tests/CQTest/Source/CQTest/Private/CQTestBlueprintHelper.cpp
deleted file mode 100644
diff --git a/Engine/Plugins/Tests/CQTest/Source/CQTest/Private/Commands/TestCommands.cpp b/Engine/Plugins/Tests/CQTest/Source/CQTest/Private/Commands/TestCommands.cpp
new file mode 100644
--- /dev/null
+++ b/Engine/Plugins/Tests/CQTest/Source/CQTest/Private/Commands/TestCommands.cpp
@@ -0,0 +1,91 @@
+// Copyright Epic Games, Inc. All Rights Reserved.
+
+#include "Commands/TestCommands.h"
+#include "Logging/LogMacros.h"
+
+DEFINE_LOG_CATEGORY_STATIC(LogCqTest, Log, All)
+
+bool CQTEST_API FWaitUntil::Update()
+{
+	if (Description != nullptr && !bHasLoggedStart)
+	{
+		UE_LOG(LogCqTest, Log, TEXT("Starting %s"), Description);
+		bHasLoggedStart = true;
+	}
+	if (TestRunner.HasAnyErrors())
+	{
+		return true;
+	}
+
+	if (Query())
+	{
+		if (Description)
+		{
+			UE_LOG(LogCqTest, Log, TEXT("Finished %s"), Description);
+		}
+		return true;
+	}
+	else if (FDateTime::UtcNow() >= StartTime + Timeout)
+	{
+		if (Description)
+		{
+			TestRunner.AddError(*FString::Printf(TEXT("Timed out waiting for %s"), Description));
+		}
+		else
+		{
+			TestRunner.AddError(TEXT("Latent command timed out."));
+		}
+		return true;
+	}
+	return false;
+}
+
+bool CQTEST_API FExecute::Update()
+{
+	if (Description)
+	{
+		UE_LOG(LogCqTest, Log, TEXT("Running %s"), Description);
+	}
+	if (!TestRunner.HasAnyErrors())
+	{
+		Func();
+	}
+	return true;
+}
+
+void CQTEST_API FRunSequence::Append(TSharedPtr<IAutomationLatentCommand> ToAdd)
+{
+	Commands.Add(ToAdd);
+}
+
+void CQTEST_API FRunSequence::AppendAll(TArray<TSharedPtr<IAutomationLatentCommand>> ToAdd)
+{
+	for (auto& Cmd : ToAdd)
+	{
+		Commands.Add(Cmd);
+	}
+}
+
+void CQTEST_API FRunSequence::Prepend(TSharedPtr<IAutomationLatentCommand> ToAdd)
+{
+	Commands.Insert(ToAdd, 0);
+}
+
+bool CQTEST_API FRunSequence::Update()
+{
+	if (Commands.Num() == 0)
+	{
+		return true;
+	}
+	else
+	{
+		auto Command = Commands[0];
+		Commands.RemoveAt(0); //Remove the command now, in case the command prepends other commands
+		if (Command != nullptr && Command->Update() == false)
+		{
+			Commands.Insert(Command, 0);
+			return false;
+		}
+		return Commands.IsEmpty();
+	}
+}
diff --git a/Engine/Plugins/Tests/CQTest/Source/CQTest/Private/Components/ActorTestSpawner.cpp b/Engine/Plugins/Tests/CQTest/Source/CQTest/Private/Components/ActorTestSpawner.cpp
new file mode 100644
--- /dev/null
+++ b/Engine/Plugins/Tests/CQTest/Source/CQTest/Private/Components/ActorTestSpawner.cpp
@@ -0,0 +1,98 @@
+// Copyright Epic Games, Inc. All Rights Reserved.
+
+#include "Components/ActorTestSpawner.h"
+#include "TestGameInstance.h"
+
+#include "EngineUtils.h"
+#include "Engine/Engine.h"
+
+namespace
+{
+	static constexpr bool bTriggerWorldDestroyCallbacks = true;
+
+	void DestroyWorld(UWorld* GameWorld)
+	{
+		if (GameWorld->AreActorsInitialized())
+		{
+			for (AActor* const Actor : FActorRange(GameWorld))
+			{
+				if (Actor)
+				{
+					Actor->RouteEndPlay(EEndPlayReason::LevelTransition);
+				}
+			}
+		}
+
+		GEngine->ShutdownWorldNetDriver(GameWorld);
+		GameWorld->DestroyWorld(bTriggerWorldDestroyCallbacks);
+		GameWorld->SetPhysicsScene(nullptr);
+		GEngine->DestroyWorldContext(GameWorld);
+	}
+
+	void DestroySpawnedActors(TArray<TWeakObjectPtr<AActor>>& SpawnedActors)
+	{
+		// Destroy actors in reverse order of creation
+		for (int32 Index = SpawnedActors.Num() - 1; Index >= 0; Index--)
+		{
+			if (AActor* CastActor = SpawnedActors[Index].Get())
+			{
+				CastActor->Destroy(true);
+			}
+		}
+
+		SpawnedActors.Empty();
+	}
+
+	void DestroySpawnedObjects(TArray<TWeakObjectPtr<UObject>>& InSpawnedObjects)
+	{
+		// Destroy objects in reverse order of creation
+		for (int32 Index = InSpawnedObjects.Num() - 1; Index >= 0; Index--)
+		{
+			if (UObject* Obj = InSpawnedObjects[Index].Get())
+			{
+				Obj->BeginDestroy();
+			}
+		}
+
+		InSpawnedObjects.Empty();
+	}
+} // namespace
+
+FActorTestSpawner::~FActorTestSpawner() {
+	if (GameInstance)
+	{
+		GameInstance->Shutdown();
+		GameInstance->RemoveFromRoot();
+		GameInstance = nullptr;
+	}
+
+	if (GameWorld)
+	{
+		DestroySpawnedActors(SpawnedActors);
+		DestroySpawnedObjects(SpawnedObjects);
+		DestroyWorld(GameWorld);
+		GameWorld->RemoveFromRoot();
+		GameWorld = nullptr;
+	}
+}
+
+UWorld* FActorTestSpawner::CreateWorld()
+{
+	FWorldContext& WorldContext = GEngine->CreateNewWorldContext(EWorldType::Game);
+	UWorld* Result = UWorld::CreateWorld(EWorldType::Game, false, NAME_None);
+	check(Result != nullptr);
+	Result->AddToRoot();
+	WorldContext.SetCurrentWorld(Result);
+
+	Result->InitializeActorsForPlay(FURL());
+	check(Result->GetPhysicsScene() != nullptr);
+
+	return Result;
+}
+
+void FActorTestSpawner::InitializeGameSubsystems()
+{
+	GameInstance = NewObject<UTestGameInstance>();
+	GameInstance->AddToRoot();
+	GameInstance->InitializeForTest(GetWorld());
+}
diff --git a/Engine/Plugins/Tests/CQTest/Source/CQTest/Private/Components/CQTestBlueprintHelper.cpp b/Engine/Plugins/Tests/CQTest/Source/CQTest/Private/Components/CQTestBlueprintHelper.cpp
new file mode 100644
--- /dev/null
+++ b/Engine/Plugins/Tests/CQTest/Source/CQTest/Private/Components/CQTestBlueprintHelper.cpp
@@ -0,0 +1,47 @@
+// Copyright Epic Games, Inc. All Rights Reserved.
+
+#include "Components/CQTestBlueprintHelper.h"
+
+#include "Misc/Paths.h"
+#include "Engine/Blueprint.h"
+#include "AssetRegistry/AssetData.h"
+#include "AssetRegistry/AssetRegistryModule.h"
+
+UClass* FCQTestBlueprintHelper::GetBlueprintClass(const FString& Directory, const FString& Name)
+{
+	const FString FileName = FString::Printf(TEXT("%s.%s"), *Name, *Name);
+	const FString Path = FPaths::Combine(*Directory, FileName);
+	UObject* BpObject = Cast<UObject>(StaticLoadObject(UObject::StaticClass(), nullptr, *Path));
+	check(BpObject);
+
+	UBlueprint* AsBlueprint = Cast<UBlueprint>(BpObject);
+	check(AsBlueprint);
+
+	return AsBlueprint->GeneratedClass;
+}
+
+UObject* FCQTestBlueprintHelper::FindDataBlueprint(const FString& Directory, const FString& Name)
+{
+	static TMap<FString, FAssetData> BlueprintCache{};
+	static TSet<FString> LoadedDirectories{};
+	if (!LoadedDirectories.Contains(Directory))
+	{
+		TArray<FAssetData> AssetData;
+		const FAssetRegistryModule& AssetRegistry = FModuleManager::LoadModuleChecked<FAssetRegistryModule>("AssetRegistry");
+
+		AssetRegistry.Get().GetAssetsByPath(FName(Directory), AssetData, true);
+		for (const auto& Data : AssetData)
+		{
+			BlueprintCache.Add(Data.GetAsset()->GetFName().ToString(), Data);
+		}
+	
+		LoadedDirectories.Add(Directory);
+	}
+
+	if (FAssetData* Data = BlueprintCache.Find(Name))
+	{
+		return Data->GetAsset();
+	}
+
+	return nullptr;
+}
\ No newline at end of file
diff --git a/Engine/Plugins/Tests/CQTest/Source/CQTest/Private/Components/MapTestSpawner.cpp b/Engine/Plugins/Tests/CQTest/Source/CQTest/Private/Components/MapTestSpawner.cpp
new file mode 100644
--- /dev/null
+++ b/Engine/Plugins/Tests/CQTest/Source/CQTest/Private/Components/MapTestSpawner.cpp
@@ -0,0 +1,47 @@
+// Copyright Epic Games, Inc. All Rights Reserved.
+
+#include "Components/MapTestSpawner.h"
+#include "Commands/TestCommands.h"
+
+#include "Tests/AutomationCommon.h"
+#include "GameFramework/PlayerController.h"
+#include "Engine/Engine.h"
+#include "Misc/Paths.h"
+
+void FMapTestSpawner::AddWaitUntilLoadedCommand(FAutomationTestBase* TestRunner)
+{
+#if WITH_AUTOMATION_TESTS
+	check(PieWorld == nullptr);
+
+	const FString FileName = FString::Printf(TEXT("%s.%s"), *MapName, *MapName);
+	const FString Path = FPaths::Combine(MapDirectory, FileName);
+	bool bOpened = AutomationOpenMap(Path);
+	check(bOpened);
+
+	ADD_LATENT_AUTOMATION_COMMAND(FWaitUntil(*TestRunner, [&]() -> bool {
+		for (const auto& Context : GEngine->GetWorldContexts())
+		{
+			if (((Context.WorldType == EWorldType::PIE) || (Context.WorldType == EWorldType::Game)) && (Context.World() != nullptr))
+			{
+				PieWorld = Context.World();
+				return true;
+			}
+		}
+
+		return false;
+	}));
+#else
+	checkf(false, TEXT("AddWaitUntilLoadedCommand can't call AutomationOpenMap if WITH_AUTOMATION_TESTS=false"));
+#endif
+}
+
+UWorld* FMapTestSpawner::CreateWorld()
+{
+	checkf(PieWorld, TEXT("Must call AddWaitUntilLoadedCommand in BEFORE_TEST"));
+	return PieWorld;
+}
+
+APawn* FMapTestSpawner::FindFirstPlayerPawn()
+{
+	return GetWorld().GetFirstPlayerController()->GetPawn();
+}
diff --git a/Engine/Plugins/Tests/CQTest/Source/CQTest/Private/Components/PIENetworkComponent.cpp b/Engine/Plugins/Tests/CQTest/Source/CQTest/Private/Components/PIENetworkComponent.cpp
new file mode 100644
--- /dev/null
+++ b/Engine/Plugins/Tests/CQTest/Source/CQTest/Private/Components/PIENetworkComponent.cpp
@@ -0,0 +1,210 @@
+// Copyright Epic Games, Inc. All Rights Reserved.
+
+#include "Components/PIENetworkComponent.h"
+
+#if ENABLE_PIE_NETWORK_TEST
+#include "LevelEditor.h"
+#include "Editor/UnrealEdEngine.h"
+#include "Settings/LevelEditorPlaySettings.h"
+#include "UnrealEdGlobals.h"
+#include "Tests/AutomationEditorCommon.h"
+
+#include "Engine/NetConnection.h"
+#include "Engine/NetDriver.h"
+#include "GameFramework/GameMode.h"
+
+#include "GameMapsSettings.h"
+
+DEFINE_LOG_CATEGORY_STATIC(LogNetworkTest, Log, All);
+
+FBasePIENetworkComponent::FBasePIENetworkComponent(FAutomationTestBase* InTestRunner, FTestCommandBuilder& InCommandBuilder, bool IsInitializing)
+	: TestRunner(InTestRunner)
+	, CommandBuilder(&InCommandBuilder)
+{
+	if (IsInitializing)
+	{
+		return;
+	}
+
+	CommandBuilder
+		->Do(TEXT("Start PIE"), [this]() { StartPie(); })
+		.Until(TEXT("Collect PIE Worlds"), [this]() { return CollectPieWorlds(); })
+		.Until(TEXT("Await Connections"), [this]() { return AwaitConnections(); });
+}
+
+FBasePIENetworkComponent::~FBasePIENetworkComponent()
+{
+	if (ServerState != nullptr)
+	{
+		GUnrealEd->RequestEndPlayMap();
+		StateRestorer.Restore();
+	}
+}
+
+FBasePIENetworkComponent& FBasePIENetworkComponent::Then(TFunction<void()> Action)
+{
+	CommandBuilder->Then(Action);
+	return *this;
+}
+
+FBasePIENetworkComponent& FBasePIENetworkComponent::Do(TFunction<void()> Action)
+{
+	CommandBuilder->Do(Action);
+	return *this;
+}
+
+FBasePIENetworkComponent& FBasePIENetworkComponent::Then(const TCHAR* Description, TFunction<void()> Action)
+{
+	CommandBuilder->Then(Description, Action);
+	return *this;
+}
+FBasePIENetworkComponent& FBasePIENetworkComponent::Do(const TCHAR* Description, TFunction<void()> Action)
+{
+	CommandBuilder->Do(Description, Action);
+	return *this;
+}
+
+FBasePIENetworkComponent& FBasePIENetworkComponent::Until(TFunction<bool()> Query, FTimespan Timeout)
+{
+	CommandBuilder->Until(Query, Timeout);
+	return *this;
+}
+FBasePIENetworkComponent& FBasePIENetworkComponent::Until(TCHAR* Description, TFunction<bool()> Query, FTimespan Timeout)
+{
+	CommandBuilder->Until(Description, Query, Timeout);
+	return *this;
+}
+
+FBasePIENetworkComponent& FBasePIENetworkComponent::StartWhen(TFunction<bool()> Query, FTimespan Timeout)
+{
+	CommandBuilder->StartWhen(Query, Timeout);
+	return *this;
+}
+FBasePIENetworkComponent& FBasePIENetworkComponent::StartWhen(TCHAR* Description, TFunction<bool()> Query, FTimespan Timeout)
+{
+	CommandBuilder->StartWhen(Description, Query, Timeout);
+	return *this;
+}
+
+void FBasePIENetworkComponent::StartPie()
+{
+	if (ServerState == nullptr)
+	{
+		TestRunner->AddError(TEXT("Failed to initialize Network Component"));
+		return;
+	}
+	GUnrealEd->RequestEndPlayMap();
+	FAutomationEditorCommonUtils::CreateNewMap();
+
+	ULevelEditorPlaySettings* PlaySettings = NewObject<ULevelEditorPlaySettings>();
+	if (ServerState->bIsDedicatedServer)
+	{
+		PlaySettings->SetPlayNetMode(EPlayNetMode::PIE_Client);
+		PlaySettings->SetPlayNumberOfClients(ServerState->ClientCount);
+	}
+	else
+	{
+		PlaySettings->SetPlayNetMode(EPlayNetMode::PIE_ListenServer);
+		PlaySettings->SetPlayNumberOfClients(ServerState->ClientCount + 1); // The listen server counts as a client, so we need to add one more to get a real client as well
+	}
+	PlaySettings->bLaunchSeparateServer = ServerState->bIsDedicatedServer;
+	PlaySettings->GameGetsMouseControl = false;
+	PlaySettings->SetRunUnderOneProcess(true);
+
+	FLevelEditorModule& LevelEditorModule = FModuleManager::Get().GetModuleChecked<FLevelEditorModule>(TEXT("LevelEditor"));
+
+	FRequestPlaySessionParams SessionParams;
+	SessionParams.WorldType = EPlaySessionWorldType::PlayInEditor;
+	SessionParams.GameModeOverride = AGameModeBase::StaticClass();
+	SessionParams.DestinationSlateViewport = LevelEditorModule.GetFirstActiveViewport();
+	SessionParams.EditorPlaySettings = PlaySettings;
+
+	GUnrealEd->RequestPlaySession(SessionParams);
+	GUnrealEd->StartQueuedPlaySessionRequest();
+}
+
+bool FBasePIENetworkComponent::CollectPieWorlds()
+{
+	const auto& WorldContexts = GEngine->GetWorldContexts();
+	UWorld* ServerWorld = nullptr;
+	TArray<UWorld*> ClientWorlds;
+
+	for (const auto& WorldContext : WorldContexts)
+	{
+		if (WorldContext.WorldType != EWorldType::PIE || WorldContext.World() == nullptr || WorldContext.World()->GetNetDriver() == nullptr)
+		{
+			continue;
+		}
+		UWorld* World = WorldContext.World();
+		if (World->GetNetDriver()->IsServer())
+		{
+			ServerWorld = World;
+			if (World->GetNetMode() == NM_DedicatedServer && !ServerState->bIsDedicatedServer)
+			{
+				TestRunner->AddError(TEXT("Failed to set up dedicated server.  Does your game's editor module override the PIE settings?"));
+
+				return true;
+			}
+			else if (World->GetNetMode() == NM_ListenServer && ServerState->bIsDedicatedServer)
+			{
+				TestRunner->AddError(TEXT("Failed to set up dedicated server.  Does your game's editor module override the PIE settings?"));
+				return true;
+			}
+		}
+		else
+		{
+			ClientWorlds.Add(World);
+		}
+	}
+	if (ServerWorld == nullptr || ClientWorlds.Num() < ServerState->ClientCount)
+	{
+		return false;
+	}
+
+	if (PacketSimulationSettings)
+	{
+		ServerWorld->GetNetDriver()->SetPacketSimulationSettings(*PacketSimulationSettings);
+		for (auto& World : ClientWorlds)
+		{
+			World->GetNetDriver()->SetPacketSimulationSettings(*PacketSimulationSettings);
+		}
+	}
+
+	ServerState->World = ServerWorld;
+	for (int32 ClientIndex = 0; ClientIndex < ServerState->ClientCount; ClientIndex++)
+	{
+		UWorld* ClientWorld = ClientWorlds[ClientIndex];
+		ClientStates[ClientIndex]->World = ClientWorld;
+
+		const int32 ClientLocalPort = ClientWorld->GetNetDriver()->GetLocalAddr()->GetPort();
+		auto ServerConnection = ServerWorld->GetNetDriver()->ClientConnections.FindByPredicate([ClientLocalPort](UNetConnection* ClientConnection) {
+			return ClientConnection->GetRemoteAddr()->GetPort() == ClientLocalPort;
+			});
+		check(ServerConnection != nullptr);
+		ServerState->ClientConnections[ClientIndex] = *ServerConnection;
+	}
+
+	return true;
+}
+
+bool FBasePIENetworkComponent::AwaitConnections()
+{
+	if (ServerState == nullptr || ServerState->World == nullptr)
+	{
+		return true; //Failed to get server state, test will fail
+	}
+	if (ServerState->World->GetNetDriver()->ClientConnections.Num() != ServerState->ClientCount)
+	{
+		return false;
+	}
+	for (UNetConnection* ClientConnection : ServerState->World->GetNetDriver()->ClientConnections)
+	{
+		if (ClientConnection->ViewTarget == nullptr)
+		{
+			return false;
+		}
+	}
+
+	return true;
+}
+#endif
\ No newline at end of file
diff --git a/Engine/Plugins/Tests/CQTest/Source/CQTest/Private/Components/PIENetworkTestStateRestorer.cpp b/Engine/Plugins/Tests/CQTest/Source/CQTest/Private/Components/PIENetworkTestStateRestorer.cpp
new file mode 100644
--- /dev/null
+++ b/Engine/Plugins/Tests/CQTest/Source/CQTest/Private/Components/PIENetworkTestStateRestorer.cpp
@@ -0,0 +1,58 @@
+// Copyright Epic Games, Inc. All Rights Reserved.
+
+#include "Components/PIENetworkTestStateRestorer.h"
+
+#include "GameMapsSettings.h"
+#include "GameFramework/GameMode.h"
+#include "GameFramework/WorldSettings.h"
+
+DEFINE_LOG_CATEGORY_STATIC(LogPieNetworkStateRestorer, Log, All);
+
+FPIENetworkTestStateRestorer::FPIENetworkTestStateRestorer(const FSoftClassPath InGameInstanceClass, const TSubclassOf<AGameModeBase> InGameMode)
+{
+	EditorWorld = GWorld;
+	if (InGameInstanceClass.IsValid())
+	{
+		UGameMapsSettings* GameMapSettings = GetMutableDefault<UGameMapsSettings>();
+		if (GameMapSettings)
+		{
+			OriginalGameInstance = GameMapSettings->GameInstanceClass;
+			GameMapSettings->GameInstanceClass = InGameInstanceClass;
+		}
+		else
+		{
+			UE_LOG(LogPieNetworkStateRestorer, Error, TEXT("Unable to get UGameMapsSettings in when creating FPIENetworkTestStateRestorer"));
+		}
+	}
+	OriginalGameMode = EditorWorld->GetWorldSettings()->DefaultGameMode;
+	if (!EditorWorld->HasAllFlags(RF_WasLoaded))
+	{
+		EditorWorld->SetFlags(RF_WasLoaded);
+		SetWasLoadedFlag = true;
+	}
+}
+
+void FPIENetworkTestStateRestorer::Restore() {
+	if (SetWasLoadedFlag)
+	{
+		EditorWorld->ClearFlags(RF_WasLoaded);
+	}
+
+	if (OriginalGameInstance.IsValid())
+	{
+		UGameMapsSettings* GameMapSettings = GetMutableDefault<UGameMapsSettings>();
+		if (GameMapSettings)
+		{
+			GameMapSettings->GameInstanceClass = OriginalGameInstance;
+		}
+		else
+		{
+			UE_LOG(LogPieNetworkStateRestorer, Error, TEXT("Unable to get UGameMapsSettings when destroying FPIENetworkTestStateRestorer"));
+		}
+	}
+
+	if (OriginalGameMode != nullptr)
+	{
+		EditorWorld->GetWorldSettings()->DefaultGameMode = OriginalGameMode;
+	}
+}
\ No newline at end of file
diff --git a/Engine/Plugins/Tests/CQTest/Source/CQTest/Private/Components/SpawnHelper.cpp b/Engine/Plugins/Tests/CQTest/Source/CQTest/Private/Components/SpawnHelper.cpp
new file mode 100644
--- /dev/null
+++ b/Engine/Plugins/Tests/CQTest/Source/CQTest/Private/Components/SpawnHelper.cpp
@@ -0,0 +1,24 @@
+// Copyright Epic Games, Inc. All Rights Reserved.
+
+#include "Components/SpawnHelper.h"
+
+FSpawnHelper::~FSpawnHelper() 
+{
+	if (GameWorld)
+	{
+		GameWorld->RemoveFromRoot();
+		GameWorld = nullptr;
+	}
+}
+
+UWorld& FSpawnHelper::GetWorld()
+{
+	if (!GameWorld)
+	{
+		GameWorld = CreateWorld();
+		check(GameWorld);
+		GameWorld->AddToRoot();
+	}
+
+	return *GameWorld;
+}
\ No newline at end of file
diff --git a/Engine/Plugins/Tests/CQTest/Source/CQTest/Private/MapTestSpawner.cpp b/Engine/Plugins/Tests/CQTest/Source/CQTest/Private/MapTestSpawner.cpp
deleted file mode 100644
diff --git a/Engine/Plugins/Tests/CQTest/Source/CQTest/Private/SpawnHelper.cpp b/Engine/Plugins/Tests/CQTest/Source/CQTest/Private/SpawnHelper.cpp
deleted file mode 100644
diff --git a/Engine/Plugins/Tests/CQTest/Source/CQTest/Public/ActorTestSpawner.h b/Engine/Plugins/Tests/CQTest/Source/CQTest/Public/ActorTestSpawner.h
deleted file mode 100644
diff --git a/Engine/Plugins/Tests/CQTest/Source/CQTest/Public/CQTest.h b/Engine/Plugins/Tests/CQTest/Source/CQTest/Public/CQTest.h
--- a/Engine/Plugins/Tests/CQTest/Source/CQTest/Public/CQTest.h
+++ b/Engine/Plugins/Tests/CQTest/Source/CQTest/Public/CQTest.h
@@ -62,12 +62,12 @@ struct TBaseTest
 	void AddCommand(TSharedPtr<IAutomationLatentCommand> Cmd);
 
 	virtual void RunTest(const FString& MethodName) = 0;
-	FTestCommandBuilder TestCommandBuilder;
-	FAutomationTestBase& TestRunner;
-	bool bHasLatentActions{ false };
+
 	bool bInitializing{ true };
 
+	FAutomationTestBase& TestRunner;
 	AsserterType Assert;
+	FTestCommandBuilder TestCommandBuilder;
 };
 
 template <typename AsserterType>
diff --git a/Engine/Plugins/Tests/CQTest/Source/CQTest/Public/CQTestBlueprintHelper.h b/Engine/Plugins/Tests/CQTest/Source/CQTest/Public/CQTestBlueprintHelper.h
deleted file mode 100644
diff --git a/Engine/Plugins/Tests/CQTest/Source/CQTest/Public/Commands/TestCommandBuilder.h b/Engine/Plugins/Tests/CQTest/Source/CQTest/Public/Commands/TestCommandBuilder.h
--- a/Engine/Plugins/Tests/CQTest/Source/CQTest/Public/Commands/TestCommandBuilder.h
+++ b/Engine/Plugins/Tests/CQTest/Source/CQTest/Public/Commands/TestCommandBuilder.h
@@ -8,33 +8,60 @@
 class FTestCommandBuilder
 {
 public:
-	FTestCommandBuilder() = default;
+	FTestCommandBuilder(FAutomationTestBase& InTestRunner)
+		: TestRunner(InTestRunner) {}
 
 	~FTestCommandBuilder()
 	{
-		check(CommandQueue.IsEmpty());
+		checkf(CommandQueue.IsEmpty(), TEXT("Adding latent actions from within latent actions is currently unsupported."));
 	}
 
-	FTestCommandBuilder& Do(TFunction<void()> Action)
+	FTestCommandBuilder& Do(const TCHAR* Description, TFunction<void()> Action)
 	{
-		CommandQueue.Add(MakeShared<FExecute>(Action));
+		if (!TestRunner.HasAnyErrors())
+		{
+			CommandQueue.Add(MakeShared<FExecute>(TestRunner, Action, Description));
+		}
 		return *this;
 	}
 
+	FTestCommandBuilder& Do(TFunction<void()> Action)
+	{
+		return Do(nullptr, Action);
+	}
+
 	FTestCommandBuilder& Then(TFunction<void()> Action)
 	{
 		return Do(Action);
 	}
 
-	FTestCommandBuilder& Until(FAutomationTestBase* TestRunner, TFunction<bool()> Query, FTimespan Timeout = FTimespan::FromSeconds(10))
+	FTestCommandBuilder& Then(const TCHAR* Description, TFunction<void()> Action)
+	{
+		return Do(Description, Action);
+	}
+
+	FTestCommandBuilder& Until(const TCHAR* Description, TFunction<bool()> Query, FTimespan Timeout = FTimespan::FromSeconds(10))
 	{
-		CommandQueue.Add(MakeShared<FWaitUntil>(TestRunner, Query, Timeout));
+		if (!TestRunner.HasAnyErrors())
+		{
+			CommandQueue.Add(MakeShared<FWaitUntil>(TestRunner, Query, Timeout, Description));
+		}
 		return *this;
 	}
 
-	FTestCommandBuilder& StartWhen(FAutomationTestBase* TestRunner, TFunction<bool()> Query, FTimespan Timeout = FTimespan::FromSeconds(10))
+	FTestCommandBuilder& Until(TFunction<bool()> Query, FTimespan Timeout = FTimespan::FromSeconds(10))
+	{
+		return Until(nullptr, Query, Timeout);
+	}
+
+	FTestCommandBuilder& StartWhen(TFunction<bool()> Query, FTimespan Timeout = FTimespan::FromSeconds(10))
 	{
-		return Until(TestRunner, Query, Timeout);
+		return Until(Query, Timeout);
+	}
+
+	FTestCommandBuilder& StartWhen(const TCHAR* Description, TFunction<bool()> Query, FTimespan Timeout = FTimespan::FromSeconds(10))
+	{
+		return Until(Description, Query, Timeout);
 	}
 
 	TSharedPtr<IAutomationLatentCommand> Build()
@@ -57,6 +84,10 @@ public:
 		return Result;
 	}
 
-private:
+protected:
 	TArray<TSharedPtr<IAutomationLatentCommand>> CommandQueue{};
+	FAutomationTestBase& TestRunner;
+
+	template<typename Asserter>
+	friend struct TBaseTest;
 };
\ No newline at end of file
diff --git a/Engine/Plugins/Tests/CQTest/Source/CQTest/Public/Commands/TestCommands.h b/Engine/Plugins/Tests/CQTest/Source/CQTest/Public/Commands/TestCommands.h
--- a/Engine/Plugins/Tests/CQTest/Source/CQTest/Public/Commands/TestCommands.h
+++ b/Engine/Plugins/Tests/CQTest/Source/CQTest/Public/Commands/TestCommands.h
@@ -1,54 +1,47 @@
 // Copyright Epic Games, Inc. All Rights Reserved.
+
 #pragma once
 
 #include "Misc/AutomationTest.h"
 
-class FWaitUntil : public IAutomationLatentCommand
+class CQTEST_API FWaitUntil : public IAutomationLatentCommand
 {
 public:
-	explicit FWaitUntil(FAutomationTestBase* TestRunner, TFunction<bool()> Query, FTimespan Timeout = FTimespan::FromSeconds(10))
-		: TestBase(TestRunner)
+	FWaitUntil(FAutomationTestBase& InTestRunner, TFunction<bool()> Query, FTimespan Timeout = FTimespan::FromSeconds(10), const TCHAR* InDescription = nullptr)
+		: TestRunner(InTestRunner)
 		, Query(MoveTemp(Query)) 
 		, Timeout(Timeout)
 		, StartTime(FDateTime::UtcNow())
+		, Description(InDescription)
 	{}
 
-	bool Update() override
-	{
-		if (Query())
-		{
-			return true;
-		}
-		else if (FDateTime::UtcNow() >= StartTime + Timeout)
-		{
-			TestBase->AddError(TEXT("Latent command timed out."));
-			return true;
-		}
-		return false;
-	}
+	bool Update() override;
 
-	FAutomationTestBase* TestBase;
+	FAutomationTestBase& TestRunner;
 	TFunction<bool()> Query;
 	FTimespan Timeout;
 	FDateTime StartTime;
+	const TCHAR* Description;
+	bool bHasLoggedStart = false;
 };
 
-class FExecute : public IAutomationLatentCommand
+class CQTEST_API FExecute : public IAutomationLatentCommand
 {
 public:
-	explicit FExecute(TFunction<void()> Func)
-		: Func(MoveTemp(Func)) {}
+	FExecute(FAutomationTestBase& InTestRunner, TFunction<void()> Func, const TCHAR* InDescription = nullptr)
+		: TestRunner(InTestRunner)
+		, Func(MoveTemp(Func)) 
+		, Description(InDescription)
+	{}
 
-	bool Update() override
-	{
-		Func();
-		return true;
-	}
+	bool Update() override;
 
+	FAutomationTestBase& TestRunner;
 	TFunction<void()> Func;
+	const TCHAR* Description = nullptr;
 };
 
-class FRunSequence : public IAutomationLatentCommand
+class CQTEST_API FRunSequence : public IAutomationLatentCommand
 {
 public:
 	FRunSequence(const TArray<TSharedPtr<IAutomationLatentCommand>>& ToAdd)
@@ -62,57 +55,15 @@ public:
 	{
 	}
 
-	bool Update() override
-	{
-		// Taking notes for the next time someone needs to understand the odd behavior here
-		// From the test framework's point of view, latent commands can all be run after the test has completed
-		// So if a BeforeTest function adds latent actions which need to happen before RunTest,
-		// then RunTest, AfterTest, and TearDown must be latent actions, run in that sequence
-		// But all of the latent actions from BeforeTest must run before RunTest is executed
-		// Which means that RunTest can't be added to the latent actions queue until
-		// there are no more latent commands in the queue to run
-		// So we need to run the first command in this series, followed by asking the test framework
-		// to execute all of its latent commands.
-		// Given that the latent actions may be waiting for tickable game objects
-		// Control must be given back to the game, so we can't use a tight loop
-		// to call ExecuteLatentCommands
-		// The next challenge is that THIS command is in that list of commands
-		// So we track if this call was recursive with the bRentrantFlag
-		if (Commands.Num() == 0)
-		{
-			return true;
-		}
-		if (bRentrantFlag)
-		{
-			return true;
-		}
-
-		if (bCommandDone)
-		{
-			bRentrantFlag = true;
-			if (FAutomationTestFramework::Get().ExecuteLatentCommands())
-			{
-				FAutomationTestFramework::Get().EnqueueLatentCommand(MakeShared<FRunSequence>(Commands));
-				return true;
-			}
-			else
-			{
-				bRentrantFlag = false;
-			}
-		}
-		else
-		{
-			if (Commands[0]->Update())
-			{
-				Commands.RemoveAt(0);
-				bCommandDone = true;
-			}
-		}
+	void Append(TSharedPtr<IAutomationLatentCommand> ToAdd);
+	void AppendAll(TArray < TSharedPtr<IAutomationLatentCommand>> ToAdd);
+	void Prepend(TSharedPtr<IAutomationLatentCommand> ToAdd);
 
-		return false;
+	bool IsEmpty() const {
+		return Commands.IsEmpty();
 	}
 
+	bool Update() override;
+
 	TArray<TSharedPtr<IAutomationLatentCommand>> Commands;
-	bool bRentrantFlag = false;
-	bool bCommandDone = false;
 };
\ No newline at end of file
diff --git a/Engine/Plugins/Tests/CQTest/Source/CQTest/Public/Components/ActorTestSpawner.h b/Engine/Plugins/Tests/CQTest/Source/CQTest/Public/Components/ActorTestSpawner.h
new file mode 100644
--- /dev/null
+++ b/Engine/Plugins/Tests/CQTest/Source/CQTest/Public/Components/ActorTestSpawner.h
@@ -0,0 +1,22 @@
+// Copyright Epic Games, Inc. All Rights Reserved.
+
+#pragma once
+
+#include "SpawnHelper.h"
+
+class UTestGameInstance;
+
+/// Class for spawning actors in an ActorTest context (no PIE loaded)
+struct CQTEST_API FActorTestSpawner : public FSpawnHelper
+{
+	FActorTestSpawner() = default;
+	virtual ~FActorTestSpawner() override;
+
+	void InitializeGameSubsystems();
+
+protected:
+	virtual UWorld* CreateWorld() override;
+
+private:
+	UTestGameInstance* GameInstance{ nullptr };
+};
\ No newline at end of file
diff --git a/Engine/Plugins/Tests/CQTest/Source/CQTest/Public/Components/CQTestBlueprintHelper.h b/Engine/Plugins/Tests/CQTest/Source/CQTest/Public/Components/CQTestBlueprintHelper.h
new file mode 100644
--- /dev/null
+++ b/Engine/Plugins/Tests/CQTest/Source/CQTest/Public/Components/CQTestBlueprintHelper.h
@@ -0,0 +1,12 @@
+// Copyright Epic Games, Inc. All Rights Reserved.
+
+#pragma once
+#include "UObject/Object.h"
+#include "UObject/Class.h"
+
+/// Utilities for working with Blueprint in C++ Tests
+struct CQTEST_API FCQTestBlueprintHelper
+{
+	UObject* FindDataBlueprint(const FString& Directory, const FString& Name);
+	UClass* GetBlueprintClass(const FString& Directory, const FString& Name);
+};
\ No newline at end of file
diff --git a/Engine/Plugins/Tests/CQTest/Source/CQTest/Public/Components/MapTestSpawner.h b/Engine/Plugins/Tests/CQTest/Source/CQTest/Public/Components/MapTestSpawner.h
new file mode 100644
--- /dev/null
+++ b/Engine/Plugins/Tests/CQTest/Source/CQTest/Public/Components/MapTestSpawner.h
@@ -0,0 +1,26 @@
+// Copyright Epic Games, Inc. All Rights Reserved.
+
+#pragma once
+
+#include "SpawnHelper.h"
+
+/// Class for spawning Actors in a named Map/Level
+struct CQTEST_API FMapTestSpawner : public FSpawnHelper
+{
+	FMapTestSpawner(const FString& MapDirectory, const FString& MapName)
+		: MapDirectory(MapDirectory)
+		, MapName(MapName)
+	{
+	}
+
+	void AddWaitUntilLoadedCommand(FAutomationTestBase* TestRunner);
+	APawn* FindFirstPlayerPawn();
+
+protected:
+    virtual UWorld* CreateWorld() override;
+
+private:
+	FString MapDirectory;
+	FString MapName;
+	UWorld* PieWorld{ nullptr };
+};
\ No newline at end of file
diff --git a/Engine/Plugins/Tests/CQTest/Source/CQTest/Public/Components/PIENetworkComponent.h b/Engine/Plugins/Tests/CQTest/Source/CQTest/Public/Components/PIENetworkComponent.h
new file mode 100644
--- /dev/null
+++ b/Engine/Plugins/Tests/CQTest/Source/CQTest/Public/Components/PIENetworkComponent.h
@@ -0,0 +1,121 @@
+// Copyright Epic Games, Inc. All Rights Reserved.
+
+#pragma once
+#include "HAL/Platform.h"
+
+#define ENABLE_PIE_NETWORK_TEST WITH_EDITOR
+
+#if ENABLE_PIE_NETWORK_TEST
+
+#include "CQTest.h"
+#include "PIENetworkTestStateRestorer.h"
+
+#include "Engine/NetDriver.h"
+#include "Engine/PackageMapClient.h"
+
+struct FBasePIENetworkComponentState
+{
+	UWorld* World = nullptr;
+	TArray<UNetConnection*> ClientConnections;
+	int32 ClientIndex = INDEX_NONE;
+	int32 ClientCount = 2;
+	bool bIsDedicatedServer = true;
+};
+
+class CQTEST_API FBasePIENetworkComponent
+{
+public:
+	FBasePIENetworkComponent(FAutomationTestBase* InTestRunner, FTestCommandBuilder& InCommandBuilder, bool IsInitializing);
+	~FBasePIENetworkComponent();
+
+	FBasePIENetworkComponent& Then(TFunction<void()> Action);
+	FBasePIENetworkComponent& Then(const TCHAR* Description, TFunction<void()> Action);
+
+	FBasePIENetworkComponent& Do(TFunction<void()> Action);
+	FBasePIENetworkComponent& Do(const TCHAR* Description, TFunction<void()> Action);
+
+	FBasePIENetworkComponent& Until(TFunction<bool()> Query, FTimespan Timeout = FTimespan::FromSeconds(10));
+	FBasePIENetworkComponent& Until(TCHAR* Description, TFunction<bool()> Query, FTimespan Timeout = FTimespan::FromSeconds(10));
+	
+	FBasePIENetworkComponent& StartWhen(TFunction<bool()> Query, FTimespan Timeout = FTimespan::FromSeconds(10));
+	FBasePIENetworkComponent& StartWhen(TCHAR* Description, TFunction<bool()> Query, FTimespan Timeout = FTimespan::FromSeconds(10));
+
+protected:
+	void StartPie();
+	bool CollectPieWorlds();
+	bool AwaitConnections();
+
+	TUniquePtr<FBasePIENetworkComponentState> ServerState = nullptr;
+	TArray<TUniquePtr<FBasePIENetworkComponentState>> ClientStates;
+	FAutomationTestBase* TestRunner = nullptr;
+	FTestCommandBuilder* CommandBuilder = nullptr;
+	FPacketSimulationSettings* PacketSimulationSettings = nullptr;
+	FPIENetworkTestStateRestorer StateRestorer;
+};
+
+template <typename FBasePIENetworkComponentState>
+class FNetworkComponentBuilder;
+
+template <typename NetworkDataType = FBasePIENetworkComponentState>
+class FPIENetworkComponent : public FBasePIENetworkComponent
+{
+public:
+	FPIENetworkComponent(FAutomationTestBase* InTestRunner, FTestCommandBuilder& InCommandBuilder, bool IsInitializing)
+		: FBasePIENetworkComponent(InTestRunner, InCommandBuilder, IsInitializing) {}
+
+	FPIENetworkComponent& ThenServer(TFunction<void(NetworkDataType&)> Action);
+	FPIENetworkComponent& ThenServer(const TCHAR* Description, TFunction<void(NetworkDataType&)> Action);
+
+	FPIENetworkComponent& ThenClients(TFunction<void(NetworkDataType&)> Action);
+	FPIENetworkComponent& ThenClients(const TCHAR* Description, TFunction<void(NetworkDataType&)> Action);
+
+	FPIENetworkComponent& ThenClient(int32 ClientIndex, TFunction<void(NetworkDataType&)> Action);
+	FPIENetworkComponent& ThenClient(const TCHAR* Description, int32 ClientIndex, TFunction<void(NetworkDataType&)> Action);
+
+	FPIENetworkComponent& UntilServer(TFunction<bool(NetworkDataType&)> Query, FTimespan Timeout = FTimespan::FromSeconds(10));
+	FPIENetworkComponent& UntilServer(const TCHAR* Description, TFunction<bool(NetworkDataType&)> Query, FTimespan Timeout = FTimespan::FromSeconds(10));
+
+	FPIENetworkComponent& UntilClients(TFunction<bool(NetworkDataType&)> Query, FTimespan Timeout = FTimespan::FromSeconds(10));
+	FPIENetworkComponent& UntilClients(const TCHAR* Description, TFunction<bool(NetworkDataType&)> Query, FTimespan Timeout = FTimespan::FromSeconds(10));
+	
+	FPIENetworkComponent& UntilClient(int32 ClientIndex, TFunction<bool(NetworkDataType&)> Query, FTimespan Timeout = FTimespan::FromSeconds(10));
+	FPIENetworkComponent& UntilClient(const TCHAR* Description, int32 ClientIndex, TFunction<bool(NetworkDataType&)> Query, FTimespan Timeout = FTimespan::FromSeconds(10));
+
+	template<typename ActorToSpawn, ActorToSpawn* NetworkDataType::*ResultStorage>
+	FPIENetworkComponent& SpawnAndReplicate();
+
+private:
+	friend class FNetworkComponentBuilder<NetworkDataType>;
+};
+
+
+template <typename NetworkDataType = FBasePIENetworkComponentState>
+class FNetworkComponentBuilder
+{
+public:
+	FNetworkComponentBuilder();
+
+	FNetworkComponentBuilder& WithClients(int32 ClientCount);
+	FNetworkComponentBuilder& AsDedicatedServer();
+	FNetworkComponentBuilder& AsListenServer();
+	FNetworkComponentBuilder& WithPacketSimulationSettings(FPacketSimulationSettings* InPacketSImulationSettings);
+	FNetworkComponentBuilder& WithGameMode(TSubclassOf<AGameModeBase> InGameMode);
+	FNetworkComponentBuilder& WithGameInstanceClass(FSoftClassPath InGameInstanceClass);
+
+	void Build(FPIENetworkComponent<NetworkDataType>& OutNetwork);
+
+private:
+	FPacketSimulationSettings* PacketSimulationSettings = nullptr;
+	TSubclassOf<AGameModeBase> GameMode = TSubclassOf<AGameModeBase>(nullptr);
+	FSoftClassPath GameInstanceClass = FSoftClassPath{};
+	int32 ClientCount = 2;
+	bool bIsDedicatedServer = true;
+};
+
+constexpr static uint32 NetworkTestContext = EAutomationTestFlags::EditorContext | EAutomationTestFlags::ProductFilter;
+
+#define NETWORK_TEST_CLASS(_ClassName, _TestDir) TEST_CLASS_WITH_FLAGS(_ClassName, _TestDir, NetworkTestContext)
+
+#include "Impl/PIENetworkComponent.inl"
+
+#endif // ENABLE_PIE_NETWORK_TEST
\ No newline at end of file
diff --git a/Engine/Plugins/Tests/CQTest/Source/CQTest/Public/Components/PIENetworkTestStateRestorer.h b/Engine/Plugins/Tests/CQTest/Source/CQTest/Public/Components/PIENetworkTestStateRestorer.h
new file mode 100644
--- /dev/null
+++ b/Engine/Plugins/Tests/CQTest/Source/CQTest/Public/Components/PIENetworkTestStateRestorer.h
@@ -0,0 +1,25 @@
+// Copyright Epic Games, Inc. All Rights Reserved.
+
+#pragma once
+#include "UObject/SoftObjectPath.h"
+#include "Templates/SubclassOf.h"
+
+class UWorld;
+class AGameModeBase;
+
+//Object which saves the state of the game before starting the test,
+//Then restores it after the test is complete.
+class CQTEST_API FPIENetworkTestStateRestorer
+{
+public:
+	FPIENetworkTestStateRestorer() = default;
+	FPIENetworkTestStateRestorer(const FSoftClassPath InGameInstanceClass, TSubclassOf<AGameModeBase> InGameMode);
+
+	void Restore();
+
+private:
+	UWorld* EditorWorld = nullptr;
+	bool SetWasLoadedFlag = false;
+	FSoftClassPath OriginalGameInstance = FSoftClassPath();
+	TSubclassOf<AGameModeBase> OriginalGameMode = nullptr;
+};
diff --git a/Engine/Plugins/Tests/CQTest/Source/CQTest/Public/Components/SpawnHelper.h b/Engine/Plugins/Tests/CQTest/Source/CQTest/Public/Components/SpawnHelper.h
new file mode 100644
--- /dev/null
+++ b/Engine/Plugins/Tests/CQTest/Source/CQTest/Public/Components/SpawnHelper.h
@@ -0,0 +1,58 @@
+// Copyright Epic Games, Inc. All Rights Reserved.
+
+#pragma once
+
+#include "Engine/World.h"
+#include "GameFramework/Actor.h"
+#include "UObject/GCObject.h"
+
+struct CQTEST_API FSpawnHelper
+{
+	virtual ~FSpawnHelper();
+
+	template <typename ActorType>
+	ActorType& SpawnActor(const FActorSpawnParameters& SpawnParameters = FActorSpawnParameters(), UClass* Class = nullptr)
+	{
+		return SpawnActorAtInWorld<ActorType>(GetWorld(), FVector::ZeroVector, FRotator::ZeroRotator, SpawnParameters, Class);
+	}
+
+	template <typename ActorType>
+	ActorType& SpawnActorAt(FVector const& Location, FRotator const& Rotation, const FActorSpawnParameters& SpawnParameters = FActorSpawnParameters(), UClass* Class = nullptr)
+	{
+		return SpawnActorAtInWorld<ActorType>(GetWorld(), Location, Rotation, SpawnParameters, Class);
+	}
+
+	template <typename ObjectType>
+	ObjectType& SpawnObject()
+	{
+		static_assert(!TIsDerivedFrom<ObjectType, AActor>::IsDerived, "Use SpawnActor to spawn AActors.");
+		static_assert(TIsDerivedFrom<ObjectType, UObject>::IsDerived, "Objects must derive from UObject.");
+		ObjectType* const Object = NewObject<ObjectType>();
+		check(Object != nullptr);
+		SpawnedObjects.Add(Object);
+		return *Object;
+	}
+
+	UWorld& GetWorld();
+
+protected:
+	virtual UWorld* CreateWorld() = 0;
+
+	TArray<TWeakObjectPtr<AActor>> SpawnedActors{};
+	TArray<TWeakObjectPtr<UObject>> SpawnedObjects{};
+
+	UWorld* GameWorld{ nullptr };
+
+private:
+	template <typename ActorType>
+	ActorType& SpawnActorAtInWorld(UWorld& World, const FVector& Location, const FRotator& Rotation, const FActorSpawnParameters& SpawnParameters = FActorSpawnParameters(), UClass* Class = nullptr)
+	{
+		static_assert(TIsDerivedFrom<ActorType, AActor>::IsDerived, "Provided type does not derive from AActor");
+
+		ActorType* const Actor = Class ? World.SpawnActor<ActorType>(Class, Location, Rotation, SpawnParameters) : World.SpawnActor<ActorType>(Location, Rotation, SpawnParameters);
+
+		check(Actor != nullptr);
+		SpawnedActors.Add(Actor);
+		return *Actor;
+	}
+};
\ No newline at end of file
diff --git a/Engine/Plugins/Tests/CQTest/Source/CQTest/Public/Impl/CQTest.inl b/Engine/Plugins/Tests/CQTest/Source/CQTest/Public/Impl/CQTest.inl
--- a/Engine/Plugins/Tests/CQTest/Source/CQTest/Public/Impl/CQTest.inl
+++ b/Engine/Plugins/Tests/CQTest/Source/CQTest/Public/Impl/CQTest.inl
@@ -2,72 +2,92 @@
 
 namespace
 {
+	DEFINE_LOG_CATEGORY_STATIC(LogCqTest, Log, All);
+
 	template <typename AsserterType>
 	struct TBeforeTestCommand : public IAutomationLatentCommand
 	{
-		explicit TBeforeTestCommand(TBaseTest<AsserterType>& CurrentTest)
-			: CurrentTest(CurrentTest) {}
+		TBeforeTestCommand(TBaseTest<AsserterType>& InCurrentTest, TSharedRef<FRunSequence> InSequence)
+			: CurrentTest(InCurrentTest)
+			, Sequence(InSequence)
+		{}
 
 		bool Update() override
 		{
+			UE_LOG(LogCqTest, Log, TEXT("Before Test"));
+
 			CurrentTest.Setup();
 			if (auto LatentActions = CurrentTest.TestCommandBuilder.Build())
 			{
-				CurrentTest.AddCommand(LatentActions);
+				Sequence->Prepend(LatentActions);
 			}
+
 			return true;
 		}
 
 		TBaseTest<AsserterType>& CurrentTest;
+		TSharedRef<FRunSequence> Sequence;
 	};
 
 	template <typename AsserterType>
 	struct TRunTestCommand : public IAutomationLatentCommand
 	{
-		TRunTestCommand(TBaseTest<AsserterType>& CurrentTest, const FString& RequestedTest, const TTestRunner<AsserterType>& TestRunner)
-			: CurrentTest(CurrentTest)
-			, RequestedTest(RequestedTest)
-			, TestRunner(TestRunner)
+		TRunTestCommand(TBaseTest<AsserterType>& InCurrentTest, const FString& InRequestedTest, const TTestRunner<AsserterType>& InTestRunner, TSharedRef<FRunSequence> InSequence)
+			: CurrentTest(InCurrentTest)
+			, RequestedTest(InRequestedTest)
+			, TestRunner(InTestRunner)
+			, Sequence(InSequence)
 		{
 		}
 
 		bool Update() override
 		{
+			UE_LOG(LogCqTest, Log, TEXT("RunTest"));
+
 			if (TestRunner.HasAnyErrors())
 			{
+				UE_LOG(LogCqTest, Log, TEXT("Skipping Test due to existing errors"));
 				return true; // skip run if errors in BeforeTest
 			}
 
 			CurrentTest.RunTest(RequestedTest);
 			if (auto LatentActions = CurrentTest.TestCommandBuilder.Build())
 			{
-				CurrentTest.AddCommand(LatentActions);
+				Sequence->Prepend(LatentActions);
 			}
+
 			return true;
 		}
 
 		TBaseTest<AsserterType>& CurrentTest;
 		const FString& RequestedTest;
 		const TTestRunner<AsserterType>& TestRunner;
+		TSharedRef<FRunSequence> Sequence;
 	};
 
 	template <typename AsserterType>
 	struct TAfterTestCommand : public IAutomationLatentCommand
 	{
-		explicit TAfterTestCommand(TBaseTest<AsserterType>& CurrentTest)
-			: CurrentTest(CurrentTest) {}
+		TAfterTestCommand(TBaseTest<AsserterType>& InCurrentTest, TSharedRef<FRunSequence> InSequence)
+			: CurrentTest(InCurrentTest)
+			, Sequence(InSequence)
+		{}
 
 		bool Update() override
 		{
+			UE_LOG(LogCqTest, Log, TEXT("Running After Test"));
+
 			CurrentTest.TearDown();
 			if (auto LatentActions = CurrentTest.TestCommandBuilder.Build())
 			{
-				CurrentTest.AddCommand(LatentActions);
+				Sequence->Prepend(LatentActions);
 			}
+
 			return true;
 		}
 
 		TBaseTest<AsserterType>& CurrentTest;
+		TSharedRef<FRunSequence> Sequence;
 	};
 
 	template <typename AsserterType>
@@ -78,6 +98,8 @@ namespace
 
 		bool Update() override
 		{
+			UE_LOG(LogCqTest, Log, TEXT("Tearing Down Test"));
+
 			TestRunner.CurrentTestPtr = nullptr;
 			return true;
 		}
@@ -91,7 +113,7 @@ template <typename AsserterType>
 inline TTestRunner<AsserterType>::TTestRunner(FString InName, int32 InLineNumber, const char* InFileName, FString InTestDir, uint32 InTestFlags, TTestInstanceGenerator<AsserterType> InFactory)
 	: FAutomationTestBase(InName, true)
 	, LineNumber(InLineNumber)
-	, FileName(InFileName)
+	, FileName(FString(InFileName))
 	, TestDir(InTestDir)
 	, TestFlags(InTestFlags)
 	, TestInstanceFactory(InFactory)
@@ -123,23 +145,22 @@ inline bool TTestRunner<AsserterType>::RunTest(const FString& RequestedTest)
 	check(CurrentTestPtr != nullptr);
 	auto& CurrentTest = *CurrentTestPtr;
 
-	auto Before = MakeShared<TBeforeTestCommand<AsserterType>>(CurrentTest);
-	auto Run = MakeShared<TRunTestCommand<AsserterType>>(CurrentTest, RequestedTest, *this);
-	auto After = MakeShared<TAfterTestCommand<AsserterType>>(CurrentTest);
+	TSharedRef<FRunSequence> CommandSequence = MakeShareable<FRunSequence>(new FRunSequence());
+	auto Before = MakeShared<TBeforeTestCommand<AsserterType>>(CurrentTest, CommandSequence);
+	auto Run = MakeShared<TRunTestCommand<AsserterType>>(CurrentTest, RequestedTest, *this, CommandSequence);
+	auto After = MakeShared<TAfterTestCommand<AsserterType>>(CurrentTest, CommandSequence);
 	auto TearDown = MakeShared<TTearDownRunner<AsserterType>>(*this);
 
 	auto RemainingSteps = TArray<TSharedPtr<IAutomationLatentCommand>>{ Before, Run, After, TearDown };
 
-	while (RemainingSteps.Num() > 0)
+	while (!RemainingSteps.IsEmpty())
 	{
 		RemainingSteps[0]->Update();
 		RemainingSteps.RemoveAt(0);
-
-		if (CurrentTestPtr != nullptr && CurrentTest.bHasLatentActions && RemainingSteps.Num() > 1)
+		if (CurrentTestPtr != nullptr && !CommandSequence->IsEmpty())
 		{
-			// Ensure that all latent commands are run before adding the next step
-			RemainingSteps.Insert(MakeShared<FExecute>([]() { return true; }), 0);
-			CurrentTest.AddCommand(new FRunSequence(RemainingSteps));
+			CommandSequence->AppendAll(RemainingSteps);
+			FAutomationTestFramework::Get().EnqueueLatentCommand(CommandSequence);
 			return true;
 		}
 	}
@@ -216,23 +237,21 @@ inline void TTest<Derived, AsserterType>::RunTest(const FString& TestName)
 
 template <typename AsserterType>
 inline TBaseTest<AsserterType>::TBaseTest(FAutomationTestBase& TestRunner, bool bInitializing)
-	: TestCommandBuilder(FTestCommandBuilder{})
+	: bInitializing(bInitializing)
 	, TestRunner(TestRunner)
-	, bInitializing(bInitializing)
 	, Assert(AsserterType{ TestRunner })
+	, TestCommandBuilder(FTestCommandBuilder{ TestRunner })
 {
 }
 
 template <typename AsserterType>
 inline void TBaseTest<AsserterType>::AddCommand(IAutomationLatentCommand* Cmd)
 {
-	TestRunner.AddCommand(Cmd);
-	bHasLatentActions = true;
+	TestCommandBuilder.CommandQueue.Add(MakeShareable<IAutomationLatentCommand>(Cmd));
 }
 
 template <typename AsserterType>
 inline void TBaseTest<AsserterType>::AddCommand(TSharedPtr<IAutomationLatentCommand> Cmd)
 {
-	FAutomationTestFramework::Get().EnqueueLatentCommand(Cmd);
-	bHasLatentActions = true;
+	TestCommandBuilder.CommandQueue.Add(Cmd);
 }
diff --git a/Engine/Plugins/Tests/CQTest/Source/CQTest/Public/Impl/PIENetworkComponent.inl b/Engine/Plugins/Tests/CQTest/Source/CQTest/Public/Impl/PIENetworkComponent.inl
new file mode 100644
--- /dev/null
+++ b/Engine/Plugins/Tests/CQTest/Source/CQTest/Public/Impl/PIENetworkComponent.inl
@@ -0,0 +1,224 @@
+// Copyright Epic Games, Inc. All Rights Reserved.
+
+template <typename NetworkDataType>
+inline FPIENetworkComponent<NetworkDataType>& FPIENetworkComponent<NetworkDataType>::ThenServer(TFunction<void(NetworkDataType&)> Action)
+{
+	CommandBuilder->Do([this, Action]() { Action(static_cast<NetworkDataType&>(*ServerState)); });		
+	return *this;
+}
+template <typename NetworkDataType>
+inline FPIENetworkComponent<NetworkDataType>& FPIENetworkComponent<NetworkDataType>::ThenServer(const TCHAR* Description, TFunction<void(NetworkDataType&)> Action)
+{
+	CommandBuilder->Do(Description, [this, Action] { Action(static_cast<NetworkDataType&>(*ServerState)); });
+	return *this;
+}
+
+template <typename NetworkDataType>
+inline FPIENetworkComponent<NetworkDataType>& FPIENetworkComponent<NetworkDataType>::ThenClients(TFunction<void(NetworkDataType&)> Action)
+{
+	for (int32 Index = 0; Index < ClientStates.Num(); Index++)
+	{
+		CommandBuilder->Do([this, Action, Index]() { Action(static_cast<NetworkDataType&>(*ClientStates[Index])); });
+	}
+	return *this;
+}
+template <typename NetworkDataType>
+inline FPIENetworkComponent<NetworkDataType>& FPIENetworkComponent<NetworkDataType>::ThenClients(const TCHAR* Description, TFunction<void(NetworkDataType&)> Action)
+{
+	for (int32 Index = 0; Index < ClientStates.Num(); Index++)
+	{
+		CommandBuilder->Do(Description, [this, Action, Index]() { Action(static_cast<NetworkDataType&>(*ClientStates[Index])); });
+	}
+	return *this;
+}
+
+template <typename NetworkDataType>
+inline FPIENetworkComponent<NetworkDataType>& FPIENetworkComponent<NetworkDataType>::ThenClient(int32 ClientIndex, TFunction<void(NetworkDataType&)> Action)
+{
+	if (!ClientStates.IsValidIndex(ClientIndex))
+	{
+		TestRunner->AddError(FString::Printf(TEXT("Invalid client index specified.  Requested Index %d out of %d"), ClientIndex, ClientStates.Num()));
+		return *this;
+	}
+
+	CommandBuilder->Do([this, Action, ClientIndex]() { Action(static_cast<NetworkDataType&>(*ClientStates[ClientIndex])); });
+	return *this;
+}
+template <typename NetworkDataType>
+inline FPIENetworkComponent<NetworkDataType>& FPIENetworkComponent<NetworkDataType>::ThenClient(const TCHAR* Description, int32 ClientIndex, TFunction<void(NetworkDataType&)> Action)
+{
+	if (!ClientStates.IsValidIndex(ClientIndex))
+	{
+		TestRunner->AddError(FString::Printf(TEXT("Invalid client index specified.  Requested Index %d out of %d"), ClientIndex, ClientStates.Num()));
+		return *this;
+	}
+
+	CommandBuilder->Do(Description, [this, Action, ClientIndex]() { Action(static_cast<NetworkDataType&>(*ClientStates[ClientIndex])); });
+	return *this;
+}
+
+template <typename NetworkDataType>
+inline FPIENetworkComponent<NetworkDataType>& FPIENetworkComponent<NetworkDataType>::UntilServer(TFunction<bool(NetworkDataType&)> Query, FTimespan Timeout)
+{
+	CommandBuilder->Until(
+		[this, Query]() { return Query(static_cast<NetworkDataType&>(*ServerState)); }, Timeout);
+	return *this;
+}
+template <typename NetworkDataType>
+inline FPIENetworkComponent<NetworkDataType>& FPIENetworkComponent<NetworkDataType>::UntilServer(const TCHAR* Description, TFunction<bool(NetworkDataType&)> Query, FTimespan Timeout)
+{
+	CommandBuilder->Until(
+		Description, [this, Query]() { return Query(static_cast<NetworkDataType&>(*ServerState)); }, Timeout);
+	return *this;
+}
+
+template <typename NetworkDataType>
+inline FPIENetworkComponent<NetworkDataType>& FPIENetworkComponent<NetworkDataType>::UntilClients(TFunction<bool(NetworkDataType&)> Query, FTimespan Timeout)
+{
+	for (int32 Index = 0; Index < ClientStates.Num(); Index++)
+	{
+		CommandBuilder->Until(
+			[this, Query, Index]() { return Query(static_cast<NetworkDataType&>(*ClientStates[Index])); }, Timeout);
+	}
+	return *this;
+}
+template <typename NetworkDataType>
+inline FPIENetworkComponent<NetworkDataType>& FPIENetworkComponent<NetworkDataType>::UntilClients(const TCHAR* Description, TFunction<bool(NetworkDataType&)> Query, FTimespan Timeout)
+{
+	for (int32 Index = 0; Index < ClientStates.Num(); Index++)
+	{
+		CommandBuilder->Until(
+			Description, [this, Query, Index]() { return Query(static_cast<NetworkDataType&>(*ClientStates[Index])); }, Timeout);
+	}
+	return *this;
+}
+template <typename NetworkDataType>
+inline FPIENetworkComponent<NetworkDataType>& FPIENetworkComponent<NetworkDataType>::UntilClient(int32 ClientIndex, TFunction<bool(NetworkDataType&)> Query, FTimespan Timeout)
+{
+	if (!ClientStates.IsValidIndex(ClientIndex))
+	{
+		TestRunner->AddError(FString::Printf(TEXT("Invalid client index specified.  Requested Index %d out of %d"), ClientIndex, ClientStates.Num()));
+		return *this;
+	}
+
+	CommandBuilder->Until(
+		[this, Query, ClientIndex]() { return Query(static_cast<NetworkDataType&>(*ClientStates[ClientIndex])); }, Timeout);
+	return *this;
+}
+template <typename NetworkDataType>
+inline FPIENetworkComponent<NetworkDataType>& FPIENetworkComponent<NetworkDataType>::UntilClient(const TCHAR* Description, int32 ClientIndex, TFunction<bool(NetworkDataType&)> Query, FTimespan Timeout)
+{
+	if (!ClientStates.IsValidIndex(ClientIndex))
+	{
+		TestRunner->AddError(FString::Printf(TEXT("Invalid client index specified.  Requested Index %d out of %d"), ClientIndex, ClientStates.Num()));
+		return *this;
+	}
+
+	CommandBuilder->Until(
+		Description, [this, Query, ClientIndex]() { return Query(static_cast<NetworkDataType&>(*ClientStates[ClientIndex])); }, Timeout);
+	return *this;
+}
+
+template <typename NetworkDataType>
+template <typename ActorToSpawn, ActorToSpawn* NetworkDataType::*ResultStorage>
+inline FPIENetworkComponent<NetworkDataType>& FPIENetworkComponent<NetworkDataType>::SpawnAndReplicate()
+{
+	static_assert(std::is_convertible_v<ActorToSpawn*, AActor*>, "ActorToSpawn must derive from AActor");
+
+	TSharedPtr<ActorToSpawn*> SharedStorage = MakeShareable(new ActorToSpawn*(nullptr));
+	TSharedRef<FNetworkGUID> SharedGuid = MakeShareable(new FNetworkGUID());
+
+	ThenServer(TEXT("Spawning Actor On Server"), [SharedStorage](NetworkDataType& State) {
+		*SharedStorage = State.World->template SpawnActor<ActorToSpawn>();
+		if (ResultStorage != nullptr)
+		{
+			State.*ResultStorage = *SharedStorage;
+		}
+	}).UntilServer(TEXT("Waiting for NetGUID"), [SharedStorage, SharedGuid](NetworkDataType& State) {
+		*SharedGuid = State.World->GetNetDriver()->GuidCache->GetNetGUID(*SharedStorage);
+		return SharedGuid->IsValid();
+	}).UntilClients(TEXT("Waiting for Replication on Clients"), [SharedGuid](NetworkDataType& State) {
+		ActorToSpawn* ClientActor = Cast<ActorToSpawn>(State.World->GetNetDriver()->GuidCache->GetObjectFromNetGUID(*SharedGuid, true));
+		if (ClientActor == nullptr)
+		{
+			return false;
+		}
+		if (ResultStorage != nullptr)
+		{
+			State.*ResultStorage = ClientActor;
+		}
+		return true;
+	});
+
+	return *this;
+}
+
+//////////////////////
+
+template <typename NetworkDataType>
+inline FNetworkComponentBuilder<NetworkDataType>::FNetworkComponentBuilder()
+{
+	static_assert(std::is_convertible_v<NetworkDataType*, FBasePIENetworkComponentState*>, "NetworkDataType must derive from FBaseNetworkComponentState");
+}
+
+template <typename NetworkDataType>
+inline FNetworkComponentBuilder<NetworkDataType>& FNetworkComponentBuilder<NetworkDataType>::WithClients(int32 InClientCount)
+{
+	checkf(InClientCount > 0, TEXT("Client count must be greater than 0.  Server only tests should simply use a Spawner"));
+	ClientCount = InClientCount;
+	return *this;
+}
+
+template<typename NetworkDataType>
+inline FNetworkComponentBuilder<NetworkDataType>& FNetworkComponentBuilder<NetworkDataType>::WithPacketSimulationSettings(FPacketSimulationSettings* InPacketSimulationSettings)
+{
+	PacketSimulationSettings = InPacketSimulationSettings;
+}
+
+template <typename NetworkDataType>
+inline FNetworkComponentBuilder<NetworkDataType>& FNetworkComponentBuilder<NetworkDataType>::AsDedicatedServer()
+{
+	bIsDedicatedServer = true;
+	return *this;
+}
+
+template <typename NetworkDataType>
+inline FNetworkComponentBuilder<NetworkDataType>& FNetworkComponentBuilder<NetworkDataType>::AsListenServer()
+{
+	bIsDedicatedServer = false;
+	return *this;
+}
+
+template <typename NetworkDataType>
+inline FNetworkComponentBuilder<NetworkDataType>& FNetworkComponentBuilder<NetworkDataType>::WithGameMode(TSubclassOf<AGameModeBase> InGameMode)
+{
+	GameMode = InGameMode;
+	return *this;
+}
+
+template <typename NetworkDataType>
+inline FNetworkComponentBuilder<NetworkDataType>& FNetworkComponentBuilder<NetworkDataType>::WithGameInstanceClass(FSoftClassPath InGameInstanceClass)
+{
+	GameInstanceClass = InGameInstanceClass;
+	return *this;
+}
+
+template<typename NetworkDataType>
+inline void FNetworkComponentBuilder<NetworkDataType>::Build(FPIENetworkComponent<NetworkDataType>& OutNetwork)
+{
+	NetworkDataType DefaultState{};
+	DefaultState.ClientCount = ClientCount;
+	DefaultState.bIsDedicatedServer = bIsDedicatedServer;
+
+	OutNetwork.ServerState = MakeUnique<NetworkDataType>(DefaultState);
+	OutNetwork.ServerState->ClientConnections.SetNum(ClientCount);
+
+	for (int32 ClientIndex = 0; ClientIndex < ClientCount; ClientIndex++)
+	{
+		OutNetwork.ClientStates.Add(MakeUnique<NetworkDataType>(DefaultState));
+		OutNetwork.ClientStates.Last()->ClientIndex = ClientIndex;
+	}
+
+	OutNetwork.PacketSimulationSettings = PacketSimulationSettings;
+	OutNetwork.StateRestorer = FPIENetworkTestStateRestorer{GameInstanceClass, GameMode};
+}
diff --git a/Engine/Plugins/Tests/CQTest/Source/CQTest/Public/MapTestSpawner.h b/Engine/Plugins/Tests/CQTest/Source/CQTest/Public/MapTestSpawner.h
deleted file mode 100644
diff --git a/Engine/Plugins/Tests/CQTest/Source/CQTest/Public/ObjectBuilder.h b/Engine/Plugins/Tests/CQTest/Source/CQTest/Public/ObjectBuilder.h
--- a/Engine/Plugins/Tests/CQTest/Source/CQTest/Public/ObjectBuilder.h
+++ b/Engine/Plugins/Tests/CQTest/Source/CQTest/Public/ObjectBuilder.h
@@ -9,7 +9,7 @@
 #include "UObject/UObjectGlobals.h"
 #include "UObject/Package.h"
 
-#include "SpawnHelper.h"
+#include "Components/SpawnHelper.h"
 
 DEFINE_LOG_CATEGORY_STATIC(LogObjectBuilder, Log, Log);
 
diff --git a/Engine/Plugins/Tests/CQTest/Source/CQTest/Public/SpawnHelper.h b/Engine/Plugins/Tests/CQTest/Source/CQTest/Public/SpawnHelper.h
deleted file mode 100644
diff --git a/Engine/Plugins/Tests/CQTest/Source/CQTestTests/Private/Assert/AssertionTests.cpp b/Engine/Plugins/Tests/CQTest/Source/CQTestTests/Private/Assert/AssertionTests.cpp
new file mode 100644
--- /dev/null
+++ b/Engine/Plugins/Tests/CQTest/Source/CQTestTests/Private/Assert/AssertionTests.cpp
@@ -0,0 +1,293 @@
+// Copyright Epic Games, Inc. All Rights Reserved.
+
+#include "CQTest.h"
+#include "Assert/NoDiscardAsserter.h"
+
+DECLARE_LOG_CATEGORY_CLASS(AssertionTests, Log, All);
+
+namespace CQTests
+{
+	TEST_CLASS(NoDiscardAssert, "TestFramework.CQTest.Core")
+	{
+		const char* ExpectedError = "Expected";
+
+		TEST_METHOD(AssertFail_WithMessage_AddsError)
+		{
+			Assert.ExpectError(TEXT(""));
+			Assert.Fail(TEXT(""));
+		}
+
+		TEST_METHOD(Assertions_Accept_RawStrings)
+		{
+			Assert.ExpectError("Hello World");
+			Assert.Fail("Hello World");
+		}
+
+		TEST_METHOD(Assertions_Accept_TCharArrays)
+		{
+			Assert.ExpectError(TEXT("Hello World"));
+			Assert.Fail(TEXT("Hello World"));
+		}
+
+		TEST_METHOD(Assertions_Accept_FStrings)
+		{
+			FString message = TEXT("Hello World");
+			Assert.ExpectError(message);
+			Assert.Fail(TEXT("Hello WorldA"));
+		}
+
+		TEST_METHOD(AssertExpectError_WithMultipleErrors_Succeeds)
+		{
+			Assert.ExpectError("", 3);
+			Assert.Fail("One");
+			Assert.Fail("Two");
+			Assert.Fail("Three");
+		}
+
+		TEST_METHOD(AsserterExpectError_WithZeroExpected_AcceptsAnyNumber)
+		{
+			Assert.ExpectError("", 0);
+			Assert.Fail("One");
+			Assert.Fail("Two");
+			Assert.Fail("Three");
+		}
+
+		TEST_METHOD(AssertExpectError_WithMatchingError_Succeeds)
+		{
+			Assert.ExpectError("Hello World");
+			Assert.Fail("Hello World");
+		}
+
+		TEST_METHOD(AssertExpectError_WithRegexSymbols_EscapesRegex)
+		{
+			Assert.ExpectError("[^abc]");
+			Assert.Fail("[^abc]");
+		}
+
+		TEST_METHOD(AssertExpectErrorRegex_WithRegex_Succeeds)
+		{
+			Assert.ExpectErrorRegex("\\w+");
+			Assert.Fail("abc");
+		}
+
+		TEST_METHOD(AssertTrue_WithTrue_Succeeds)
+		{
+			ASSERT_THAT(IsTrue(true));
+		}
+		TEST_METHOD(AssertTrue_WithTrueAndErrorMessage_DoesNotAddErrorMessage)
+		{
+			ASSERT_THAT(IsTrue(true, "Unexpected"));
+		}
+		TEST_METHOD(AssertTrue_WithFalse_AddsError)
+		{
+			Assert.ExpectError(TEXT(""));
+			ASSERT_THAT(IsTrue(false));
+		}
+		TEST_METHOD(AssertTrue_WithFalseAndError_AddSpecificError)
+		{
+			Assert.ExpectError(FString(ExpectedError));
+			ASSERT_THAT(IsTrue(false, ExpectedError));
+		}
+
+		TEST_METHOD(AssertEqual_WithSameInts_Succeeds)
+		{
+			ASSERT_THAT(AreEqual(42, 42));
+		}
+		TEST_METHOD(AssertEqual_WithSameIntsAndErrorMessage_DoesNotAddErrorMessage)
+		{
+			ASSERT_THAT(AreEqual(42, 42, "Unexpected"));
+		}
+		TEST_METHOD(AssertEqual_WithDifferentInts_AddsError)
+		{
+			Assert.ExpectError(TEXT(""));
+			ASSERT_THAT(AreEqual(42, 0));
+		}
+		TEST_METHOD(AssertEqual_WithDifferentIntsAndErrorMessage_AddsSpecificError)
+		{
+			Assert.ExpectError(FString(ExpectedError));
+			ASSERT_THAT(AreEqual(42, 0, ExpectedError));
+		}
+
+		TEST_METHOD(AssertEqual_WithSameFStrings_Succeeds)
+		{
+			ASSERT_THAT(AreEqual(TEXT("Hello"), TEXT("Hello")));
+		}
+		TEST_METHOD(AssertEqual_WithDifferentFStrings_AddsError)
+		{
+			Assert.ExpectError(TEXT(""));
+			ASSERT_THAT(AreEqual(TEXT("Hello"), TEXT("World")));
+		}
+		TEST_METHOD(AssertEqual_WithDifferentCases_AddsError)
+		{
+			Assert.ExpectError(TEXT(""));
+			ASSERT_THAT(AreEqual(TEXT("hello"), TEXT("HELLO")));
+		}
+		TEST_METHOD(AssertEqual_WithDifferentFStringsAndErrorMessage_AddsSpecificError)
+		{
+			Assert.ExpectError(FString(ExpectedError));
+			ASSERT_THAT(AreEqual(TEXT("Hello"), TEXT("World"), ExpectedError));
+		}
+		TEST_METHOD(AssertNear_WithSameNumbers_Succeeds)
+		{
+			ASSERT_THAT(IsNear(3.14, 3.14, 0.001));
+		}
+		TEST_METHOD(AssertNear_WithSameNumbersAndErrorMessage_DoesNotAddErrorMessage)
+		{
+			ASSERT_THAT(IsNear(3.14, 3.14, 0.001, "Unexpected"));
+		}
+		TEST_METHOD(AssertNear_WithSimilarNumbers_Succeeds)
+		{
+			ASSERT_THAT(IsNear(3.0, 3.1, 1.0));
+		}
+		TEST_METHOD(AssertNear_WithDifferentNumbers_AddsError)
+		{
+			Assert.ExpectError(TEXT(""));
+			ASSERT_THAT(IsNear(1.0, 2.0, 0.001));
+		}
+		TEST_METHOD(AssertNear_WithDifferentNumbersAndError_AddsSpecificError)
+		{
+			Assert.ExpectError(FString(ExpectedError));
+			ASSERT_THAT(IsNear(1.0, 2.0, 0.001, ExpectedError));
+		}
+
+		TEST_METHOD(AssertNull_WithNullPtr_Succeeds)
+		{
+			int* ptr = nullptr;
+			ASSERT_THAT(IsNull(ptr));
+		}
+		TEST_METHOD(AssertNull_WithNullPtrAndErrorMessage_DoesNotAddError)
+		{
+			int* ptr = nullptr;
+			ASSERT_THAT(IsNull(ptr, "Unexpected"));
+		}
+		TEST_METHOD(AssertNull_WithNonNull_AddsError)
+		{
+			int val = 42;
+			int* ptr = &val;
+			Assert.ExpectError(TEXT(""));
+			ASSERT_THAT(IsNull(ptr));
+		}
+		TEST_METHOD(AssertNull_WithNonNullAndErrorMessage_AddsSpecificErrorMessage)
+		{
+			int val = 42;
+			int* ptr = &val;
+			Assert.ExpectError(FString(ExpectedError));
+			ASSERT_THAT(IsNull(ptr, ExpectedError));
+		}
+		TEST_METHOD(AssertNull_WithInvalidSharedPtr_Succeeds)
+		{
+			auto ptr = TSharedPtr<int>{};
+			ASSERT_THAT(IsNull(ptr));
+		}
+		TEST_METHOD(AssertNull_WithValidSharedPtr_AddsError)
+		{
+			TSharedPtr<int> ptr = MakeShared<int>(42);
+			Assert.ExpectError(TEXT(""));
+			ASSERT_THAT(IsNull(ptr));
+		}
+		TEST_METHOD(AssertNull_WithInvalidUniquePtr_Succeeds)
+		{
+			TUniquePtr<int> ptr;
+			ASSERT_THAT(IsNull(ptr));
+		}
+		TEST_METHOD(AssertNull_WithValidUniquePtr_AddsError)
+		{
+			TUniquePtr<int> ptr = MakeUnique<int>(42);
+			Assert.ExpectError(TEXT(""));
+			ASSERT_THAT(IsNull(ptr));
+		}
+
+		TEST_METHOD(AssertNotNull_WithNullPtr_AddsError)
+		{
+			int* ptr = nullptr;
+			Assert.ExpectError(TEXT(""));
+			ASSERT_THAT(IsNotNull(ptr));
+		}
+		TEST_METHOD(AssertNotNull_WithNullPtrAndError_AddsSpecificError)
+		{
+			int* ptr = nullptr;
+			Assert.ExpectError(FString(ExpectedError));
+			ASSERT_THAT(IsNotNull(ptr, ExpectedError));
+		}
+		TEST_METHOD(AssertNotNull_WithNonNull_Succeeds)
+		{
+			int val = 42;
+			int* ptr = &val;
+			ASSERT_THAT(IsNotNull(ptr));
+		}
+		TEST_METHOD(AssertNotNull_WithNonNullAndErrorMessage_DoesNotAddErrorMessage)
+		{
+			int val = 42;
+			int* ptr = &val;
+			ASSERT_THAT(IsNotNull(ptr, "Unexpected"));
+		}
+		TEST_METHOD(AssertNotNull_WithInvalidSharedPtr_AddsError)
+		{
+			auto ptr = TSharedPtr<int>{};
+			Assert.ExpectError(TEXT(""));
+			ASSERT_THAT(IsNotNull(ptr));
+		}
+		TEST_METHOD(AssertNotNull_WithValidSharedPtr_Succeeds)
+		{
+			TSharedPtr<int> ptr = MakeShared<int>(42);
+			ASSERT_THAT(IsNotNull(ptr));
+		}
+		TEST_METHOD(AssertNotNull_WithInvalidUniquePtr_AddsError)
+		{
+			TUniquePtr<int> ptr;
+			Assert.ExpectError(TEXT(""));
+			ASSERT_THAT(IsNotNull(ptr));
+		}
+		TEST_METHOD(AssertNotNull_WithValidUniquePtr_Succeeds)
+		{
+			TUniquePtr<int> ptr = MakeUnique<int>(42);
+			ASSERT_THAT(IsNotNull(ptr));
+		}
+	};
+	
+} // namespace CQTests
+
+struct FCustomType
+{
+	FCustomType(const FString& InName)
+		: Name(InName)
+	{
+	}
+
+	FString Name;
+
+	bool operator==(const FCustomType& other) const
+	{
+		return Name == other.Name;
+	}
+	bool operator!=(const FCustomType& other) const
+	{
+		return !(*this == other);
+	}
+};
+
+template <>
+FString CQTestConvert::ToString(const FCustomType& obj)
+{
+	return obj.Name;
+}
+
+TEST_CLASS(AssertMessage, "TestFramework.CQTest.Core")
+{
+	TEST_METHOD(AssertAreEqual_WithUnequalCustomTypes_PrintObjectNames)
+	{
+		FCustomType left = { "left" };
+		FCustomType right = { "right" };
+
+		Assert.ExpectError("Expected left to equal right");
+		ASSERT_THAT(AreEqual(left, right));
+	}
+
+	TEST_METHOD(AssertAreNotEqual_WithEqualCustomTypes_PrintObjectNames)
+	{
+		FCustomType obj = { "Object" };
+
+		Assert.ExpectError("Expected Object to not equal Object");
+		ASSERT_THAT(AreNotEqual(obj, obj));
+	}
+};
diff --git a/Engine/Plugins/Tests/CQTest/Source/CQTestTests/Private/Assert/CQTestConvertTests.cpp b/Engine/Plugins/Tests/CQTest/Source/CQTestTests/Private/Assert/CQTestConvertTests.cpp
new file mode 100644
--- /dev/null
+++ b/Engine/Plugins/Tests/CQTest/Source/CQTestTests/Private/Assert/CQTestConvertTests.cpp
@@ -0,0 +1,202 @@
+// Copyright Epic Games, Inc. All Rights Reserved.
+
+#include "CQTest.h"
+
+static const FString OverrideStr = FString("Override");
+static const FString ToStringStr = FString("ToString");
+constexpr auto OStreamStr = "OStream";
+
+struct TypeWithExplicitOverride
+{
+};
+
+template <>
+inline FString CQTestConvert::ToString(const TypeWithExplicitOverride&)
+{
+	return OverrideStr;
+}
+
+struct TypeWithToString
+{
+	FString ToString() const
+	{
+		return ToStringStr;
+	}
+};
+
+struct TypeWithOStream
+{
+};
+
+CQTestConvert::PrintStream& operator<<(CQTestConvert::PrintStream& stream, const TypeWithOStream& obj) {
+	stream << OStreamStr;
+	return stream;
+}
+
+struct TypeWithToStringAndOverride
+{
+	FString ToString() const {
+		return ToStringStr;
+	}
+};
+
+template<>
+inline FString CQTestConvert::ToString(const TypeWithToStringAndOverride&) {
+	return OverrideStr;
+}
+
+struct TypeWithToStringAndOStream
+{
+	FString ToString() const {
+		return ToStringStr;
+	}
+	CQTestConvert::PrintStream& operator<<(CQTestConvert::PrintStream& stream)
+	{
+		stream << OStreamStr;
+		return stream;
+	}
+};
+
+enum UnscopedEnumNoOperator
+{
+	One = 1,
+	Two = 2
+};
+
+enum struct ScopedEnumNoOperator
+{
+	One = 1,
+	Two = 2
+};
+
+enum UnscopedEnumWithOperator
+{
+	Three = 3,
+	Four = 4
+};
+
+CQTestConvert::PrintStream& operator<<(CQTestConvert::PrintStream& stream, const UnscopedEnumWithOperator& e)
+{
+	switch (e)
+	{
+		case UnscopedEnumWithOperator::Three:
+			stream << "Three";
+			break;
+		case UnscopedEnumWithOperator::Four:
+			stream << "Four";
+			break;
+		default:
+			stream << "Unknown";
+			break;
+	}
+	return stream;
+}
+
+enum struct ScopedEnumWithOperator
+{
+	One = 1,
+	Two = 2
+};
+
+CQTestConvert::PrintStream& operator<<(CQTestConvert::PrintStream& stream, const ScopedEnumWithOperator& e)
+{
+	switch (e)
+	{
+		case ScopedEnumWithOperator::One:
+			stream << "One";
+			break;
+		case ScopedEnumWithOperator::Two:
+			stream << "Two";
+			break;
+		default:
+			stream << "Unknown";
+			break;
+	}
+	return stream;
+}
+
+TEST_CLASS(TestConvertToString, "TestFramework.CQTest.Core")
+{
+	TEST_METHOD(TypeWithExplicitOverride_ConvertToString_ReturnsOverride)
+	{
+		TypeWithExplicitOverride obj;
+		ASSERT_THAT(AreEqual(OverrideStr, CQTestConvert::ToString(obj)));
+	}
+
+	TEST_METHOD(TypeWithToString_ConvertToString_ReturnsToString)
+	{
+		TypeWithToString obj;
+		ASSERT_THAT(AreEqual(ToStringStr, CQTestConvert::ToString(obj)));
+	}
+
+	TEST_METHOD(TypeWithOstream_ConvertToString_ReturnsOStream)
+	{
+		TypeWithOStream obj;
+		ASSERT_THAT(AreEqual(FString(OStreamStr), CQTestConvert::ToString(obj)));
+	}
+
+	TEST_METHOD(TypeWithToStringAndOverride_ConvertToString_ReturnsOverride)
+	{
+		TypeWithToStringAndOverride obj;
+		ASSERT_THAT(AreEqual(OverrideStr, CQTestConvert::ToString(obj)));
+	}
+
+	TEST_METHOD(TypeWithToStringAndOStream_ConvertToString_ReturnsToString)
+	{
+		TypeWithToStringAndOStream obj;
+		ASSERT_THAT(AreEqual(ToStringStr, CQTestConvert::ToString(obj)));
+	}
+
+	TEST_METHOD(UnscopedEnumNoOperator_ConvertToString_ReturnsNumber)
+	{
+		auto e = UnscopedEnumNoOperator::One;
+		ASSERT_THAT(AreEqual(FString(TEXT("1")), CQTestConvert::ToString(e)));
+	}
+
+	TEST_METHOD(ScopedEnumNoOperator_ConvertToString_ReturnsNumber)
+	{
+		auto e = ScopedEnumNoOperator::Two;
+		ASSERT_THAT(AreEqual(FString(TEXT("2")), CQTestConvert::ToString(e)));
+	}
+
+	TEST_METHOD(UnscopedEnumWithOperator_ConvertToString_ReturnsString)
+	{
+		auto e = UnscopedEnumWithOperator::Three;
+		ASSERT_THAT(AreEqual(FString(TEXT("Three")), CQTestConvert::ToString(e)));
+	}
+
+	TEST_METHOD(ScopedEnumWithOperator_ConvertToString_ReturnsString)
+	{
+		auto e = ScopedEnumWithOperator::Two;
+		ASSERT_THAT(AreEqual(FString(TEXT("Two")), CQTestConvert::ToString(e)));
+	}
+
+	TEST_METHOD(ArrayOfValues_ConvertToString_ReturnsCommaDelimitedList)
+	{
+		TArray<int32> values = { 1, 2, 3 };
+		ASSERT_THAT(AreEqual(FString(TEXT("[1, 2, 3]")), CQTestConvert::ToString(values)));
+	}
+
+	TEST_METHOD(SetOfValues_ConvertToString_ReturnsCommaDelimitedList)
+	{
+		TSet<int32> values = { 1, 2, 3 };
+		auto asStr = CQTestConvert::ToString(values);
+		ASSERT_THAT(IsTrue(asStr.Contains(TEXT("1"))));
+		ASSERT_THAT(IsTrue(asStr.Contains(TEXT("2"))));
+		ASSERT_THAT(IsTrue(asStr.Contains(TEXT("3"))));
+	}
+
+	TEST_METHOD(MapOfValues_ConvertToString_ReturnsKeysAndValues)
+	{
+		TMap<int32, FString> values = { 
+			{ 1, FString(TEXT("One")) },
+			{ 2, FString(TEXT("Two")) },
+			{ 3, FString(TEXT("Three")) } 
+		};
+
+		auto asStr = CQTestConvert::ToString(values);
+		ASSERT_THAT(IsTrue(asStr.Contains(TEXT("1:One"))));
+		ASSERT_THAT(IsTrue(asStr.Contains(TEXT("2:Two"))));
+		ASSERT_THAT(IsTrue(asStr.Contains(TEXT("3:Three"))));
+	}
+};
diff --git a/Engine/Plugins/Tests/CQTest/Source/CQTestTests/Private/AssertionTests.cpp b/Engine/Plugins/Tests/CQTest/Source/CQTestTests/Private/AssertionTests.cpp
deleted file mode 100644
diff --git a/Engine/Plugins/Tests/CQTest/Source/CQTestTests/Private/CQTestConvertTests.cpp b/Engine/Plugins/Tests/CQTest/Source/CQTestTests/Private/CQTestConvertTests.cpp
deleted file mode 100644
diff --git a/Engine/Plugins/Tests/CQTest/Source/CQTestTests/Private/CQTestTests.cpp b/Engine/Plugins/Tests/CQTest/Source/CQTestTests/Private/CQTestTests.cpp
--- a/Engine/Plugins/Tests/CQTest/Source/CQTestTests/Private/CQTestTests.cpp
+++ b/Engine/Plugins/Tests/CQTest/Source/CQTestTests/Private/CQTestTests.cpp
@@ -1,17 +1,18 @@
 // Copyright Epic Games, Inc. All Rights Reserved.
 
 #include "CQTest.h"
+#include "CQTestUnitTestHelper.h"
 
 #include "Tickable.h"
 
 namespace CQTestTests
 {
-	TEST(MinimalTest, "TestFramework.CQTest")
+	TEST(Minimal, "TestFramework.CQTest")
 	{
 		ASSERT_THAT(IsTrue(true));
 	};
 
-	TEST_CLASS(From, "TestFramework.CQTest.GenerateDirectory.TokenProduces.[GenerateTestDirectory]")
+	TEST_CLASS(From, "TestFramework.CQTest.Core.GenerateDirectory.TokenProduces.[GenerateTestDirectory]")
 	{
 		TEST_METHOD(FolderStructure)
 		{
@@ -20,7 +21,7 @@ namespace CQTestTests
 		}
 	};
 
-	TEST_CLASS(Produces, "TestFramework.CQTest.GenerateDirectory")
+	TEST_CLASS(Produces, "TestFramework.CQTest.Core.GenerateDirectory")
 	{
 		TEST_METHOD(WithPlugins_AppearsInPlugins)
 		{
@@ -41,7 +42,7 @@ namespace CQTestTests
 		}
 	};
 
-	TEST_CLASS(SourceAndFile, "TestFramework.CQTest")
+	TEST_CLASS(SourceAndFile, "TestFramework.CQTest.Core")
 	{
 		TEST_METHOD(SetsSourceFile)
 		{
@@ -54,7 +55,7 @@ namespace CQTestTests
 		}
 	};
 
-	TEST_CLASS(DefaultFixtureTestFlags, "TestFramework.CQTest")
+	TEST_CLASS(DefaultFixtureTestFlags, "TestFramework.CQTest.Core")
 	{
 		TEST_METHOD(SetsApplicationContextMask)
 		{
@@ -67,7 +68,7 @@ namespace CQTestTests
 		}
 	};
 
-	TEST_CLASS_WITH_FLAGS(OverrideFixtureTestFlags, "TestFramework.CQTest", EAutomationTestFlags::ApplicationContextMask | EAutomationTestFlags::SmokeFilter)
+	TEST_CLASS_WITH_FLAGS(OverrideFixtureTestFlags, "TestFramework.CQTest.Core", EAutomationTestFlags::ApplicationContextMask | EAutomationTestFlags::SmokeFilter)
 	{
 		TEST_METHOD(GetTestFlags_ReturnsSetAutomationTestFlags)
 		{
@@ -75,7 +76,7 @@ namespace CQTestTests
 		}
 	};
 
-	TEST_CLASS(TestFixtureTest, "TestFramework.CQTest")
+	TEST_CLASS(TestFixtureTest, "TestFramework.CQTest.Core")
 	{
 		bool SetupCalled = false;
 		bool ShouldAddErrorDuringTearDown = false;
@@ -131,7 +132,7 @@ namespace CQTestTests
 		}
 	};
 
-	TEST_CLASS(TestFixtureConstructor, "TestFramework.CQTest")
+	TEST_CLASS(TestFixtureConstructor, "TestFramework.CQTest.Core")
 	{
 		bool SetupCalled = false;
 		uint32 SomeNumber = 0;
@@ -159,28 +160,8 @@ namespace CQTestTests
 		}
 	};
 
-	static void ClearExpectedError(FAutomationTestBase& TestRunner, const FString& ExpectedError)
-	{
-		FAutomationTestExecutionInfo testInfo;
-		TestRunner.GetExecutionInfo(testInfo);
-		if (testInfo.GetErrorTotal() != 1)
-		{
-			return;
-		}
-		testInfo.RemoveAllEvents([&ExpectedError](FAutomationEvent& event) {
-			return event.Message.Equals(ExpectedError);
-		});
-		if (testInfo.GetErrorTotal() == 0)
-		{
-			TestRunner.ClearExecutionInfo();
-		}
-	}
-	//This test is verifying that if an assertion is raised in BEFORE_EACH then the TEST_METHOD does not run
-	//But the AFTER_EACH still is run.  This is done by Asserting with an expected message in BEFORE and Asserting
-	//with an unexpected message in TEST_METHOD.  The AFTER_EACH checks for the expected message, and clears the errors
-	//if an only if there was 1 error message with the expected value
 
-	TEST_CLASS(TestAssertionInBefore, "TestFramework.CQTest")
+	TEST_CLASS(TestAssertionInBefore, "TestFramework.CQTest.Core")
 	{
 		FString ExpectedError = TEXT("Expected Error Message");
 
@@ -225,7 +206,7 @@ namespace CQTestTests
 		int32 CurrentCount{ 0 };
 	};
 
-	TEST_CLASS(LatentCommandTest, "TestFramework.CQTest")
+	TEST_CLASS(AddCommandTests, "TestFramework.CQTest.Core")
 	{
 		void SetExpectedExecutedCommandsCount(int32 count)
 		{
@@ -270,173 +251,154 @@ namespace CQTestTests
 		}
 	};
 
-	// --------------------------------------------------------
-	// Tickable Game Objects Tick
-	// --------------------------------------------------------
-	struct FTestTickable : public FTickableGameObject
+	TEST_CLASS(LatentActionsTest, "TestFramework.CQTest.Core")
 	{
-		virtual TStatId GetStatId() const override
-		{
-			return TStatId();
-		}
-
-		virtual void Tick(float DeltaTime) override
-		{
-			TickCount++;
-		}
-
-		virtual bool IsTickableInEditor() const override
-		{
-			return true;
-		}
-
-		void ResetTickCount()
-		{
-			TickCount = 0;
-		}
+		TArray<FString> CommandLog;
 
-		uint32 TickCount = 0;
-	};
+		TFunction<bool(LatentActionsTest*)> Assertion;
+		TArray<FString> KnownStrings = {
+			TEXT("One"),
+			TEXT("Two"),
+			TEXT("Three"),
+			TEXT("Four")
+		};
 
-	TEST_CLASS(GameObjectsTickTest, "TestFramework.CQTest")
-	{
-		BEFORE_EACH()
-		{
-			Tickable.ResetTickCount();
-			AddCommand(new FWaitUntil(TestRunner, [&]() { return Tickable.TickCount > 2; }));
-		}
 		AFTER_EACH()
 		{
-			ASSERT_THAT(IsTrue(Tickable.TickCount > 2));
+			ASSERT_THAT(IsTrue(Assertion(this)));
 		}
 
-		FTestTickable Tickable{};
-
-		TEST_METHOD(TestWithTickableGameObject_WaitingForTicksInSetup_WillAllowGameObjectToTick)
+		TEST_METHOD(Do_OnCommandBuilder_AddsLatentCommand)
 		{
-			TestCommandBuilder.Do([this]() { ASSERT_THAT(IsTrue(Tickable.TickCount > 2)); });
+			TestCommandBuilder.Do([&]() { CommandLog.Add(KnownStrings[0]); });
+			ASSERT_THAT(IsTrue(CommandLog.IsEmpty()));
+			Assertion = [](LatentActionsTest* test) {
+				return !test->CommandLog.IsEmpty();
+			};
 		}
 
-		TEST_METHOD(TestWithTickableGameObject_WaitingForTicksInSetup_WillBeCompleteDuringRunStep)
+		TEST_METHOD(MultipleDoCalls_OnCommandBuilder_AddsAllCommands)
 		{
-			ASSERT_THAT(IsTrue(Tickable.TickCount > 2));
+			TestCommandBuilder.Do([&]() {
+				CommandLog.Add(KnownStrings[0]);
+			})
+			.Do([&]() {
+				CommandLog.Add(KnownStrings[1]);
+			})
+			.Do([&]() {
+				CommandLog.Add(KnownStrings[2]);
+			});
+
+			Assertion = [](LatentActionsTest* test) {
+				if (test->CommandLog.Num() != 3)
+				{
+					return false;
+				}
+				for (int32 Index = 0; Index < 3; Index++)
+				{
+					if (test->CommandLog[Index] != test->KnownStrings[Index])
+					{
+						return false;
+					}
+				}
+				return true;
+			};
+		}
+
+		TEST_METHOD(DoAndAddCommand_InTheSameTest_AddCommandsInOrder)
+		{
+			TestCommandBuilder.Do([&]() { CommandLog.Add(KnownStrings[0]); });
+			AddCommand(new FExecute(*TestRunner, [&]() { CommandLog.Add(KnownStrings[1]); }));
+			TestCommandBuilder.Do([&]() { CommandLog.Add(KnownStrings[2]); });
+			AddCommand(new FExecute(*TestRunner, [&]() { CommandLog.Add(KnownStrings[3]); }));
+
+			Assertion = [](LatentActionsTest* test) {
+				if (test->CommandLog.Num() != 4)
+				{
+					return false;
+				}
+				for (int32 Index = 0; Index < 4; Index++)
+				{
+					if (test->CommandLog[Index] != test->KnownStrings[Index])
+					{
+						return false;
+					}
+				}
+				return true;
+			};
 		}
 	};
 
-	template <typename Test>
-	class FDelayedCommand : public IAutomationLatentCommand
+	TEST_CLASS(LatentActionErrors, "TestFramework.CQTest.Core")
 	{
-	public:
-		explicit FDelayedCommand(Test* InTest, int32 RequiredTicks)
-			: ExecutingTest(InTest) 
-			, RemainingTicks(RequiredTicks)
-		{}
+		FString ExpectedError = TEXT("ExpectedError");
 
-		bool Update() override
+		AFTER_EACH()
 		{
-			RemainingTicks--;
-			if (RemainingTicks == 0)
-			{
-				ExecutingTest->DelayedCommandExecuted = true;
-				return true;			
-			}
-
-			return false;
+			ClearExpectedError(*TestRunner, ExpectedError);
 		}
 
-		Test* ExecutingTest;
-		int32 RemainingTicks;
-	};
-
-	template <typename Test>
-	class FQueueCommand : public IAutomationLatentCommand
-	{
-	public:
-		explicit FQueueCommand(Test* InTest, int32 RequiredTicks)
-			: ExecutingTest(InTest) 
-			, RemainingTicks(RequiredTicks)
-			, RequiredTicks(RequiredTicks)
-		{}
-
-		bool Update() override
+		TEST_METHOD(Assertion_InLatentActions_PreventsAdditionalLatentActions)
 		{
-			RemainingTicks--;
-			if (RemainingTicks == 0) {
-				ExecutingTest->AddCommand(new FDelayedCommand(ExecutingTest, RequiredTicks));
-				return true;
-			}
-			return false;
+			TestCommandBuilder.Do([&]() {
+				Assert.Fail(ExpectedError);
+			})
+			.Then([&]() {
+				Assert.Fail("Unexpected Error");
+			});
 		}
-
-		Test* ExecutingTest;
-		int32 RemainingTicks;
-		int32 RequiredTicks;
 	};
 
-	TEST_CLASS(ChainedCommandsTest, "TestFramework.CQTest")
+	// --------------------------------------------------------
+	// Tickable Game Objects Tick
+	// --------------------------------------------------------
+	struct FTestTickable : public FTickableGameObject
 	{
-		bool DelayedCommandExecuted{ false };
-
-		BEFORE_EACH()
+		virtual TStatId GetStatId() const override
 		{
-			AddCommand(new FQueueCommand(this, 1));
+			return TStatId();
 		}
 
-		AFTER_EACH()
+		virtual void Tick(float DeltaTime) override
 		{
-			ASSERT_THAT(IsTrue(DelayedCommandExecuted));
-
-			DelayedCommandExecuted = false;
-			TSharedPtr<FQueueCommand<ChainedCommandsTest>> PostAfterEach = MakeShared<FQueueCommand<ChainedCommandsTest>>(this, 1);
-			TSharedPtr<FExecute> Assertion = MakeShared<FExecute>([&]() {
-				ASSERT_THAT(IsTrue(DelayedCommandExecuted));
-			});
-			AddCommand(new FRunSequence(PostAfterEach, Assertion));
+			TickCount++;
 		}
 
-		TEST_METHOD(DelayedCommandExecuted_WhenAddedInBeforeEach_IsTrueDuringRunTest)
+		virtual bool IsTickableInEditor() const override
 		{
-			ASSERT_THAT(IsTrue(DelayedCommandExecuted));
+			return true;
 		}
 
-		TEST_METHOD(DelayedCommandExecuted_WhenAddedDuringTest_IsTrueAfterTest)
+		void ResetTickCount()
 		{
-			DelayedCommandExecuted = false;
-			AddCommand(new FQueueCommand(this, 1));
+			TickCount = 0;
 		}
+
+		uint32 TickCount = 0;
 	};
 
-	TEST_CLASS(MultiTickChainedCommandsTest, "TestFramework.CQTest")
+	TEST_CLASS(GameObjectsTickTest, "TestFramework.CQTest.Core")
 	{
-		bool DelayedCommandExecuted{ false };
-
 		BEFORE_EACH()
 		{
-			AddCommand(new FQueueCommand(this, 3));
+			Tickable.ResetTickCount();
+			AddCommand(new FWaitUntil(*TestRunner, [&]() { return Tickable.TickCount > 2; }));
 		}
-
 		AFTER_EACH()
 		{
-			ASSERT_THAT(IsTrue(DelayedCommandExecuted));
-			DelayedCommandExecuted = false;
-
-			TSharedPtr<FQueueCommand<MultiTickChainedCommandsTest>> PostAfterEach = MakeShared<FQueueCommand<MultiTickChainedCommandsTest>>(this, 3);
-			TSharedPtr<FExecute> Assertion = MakeShared<FExecute>([&]() {
-				ASSERT_THAT(IsTrue(DelayedCommandExecuted));
-			});
-
-			AddCommand(new FRunSequence(PostAfterEach, Assertion));
+			ASSERT_THAT(IsTrue(Tickable.TickCount > 2));
 		}
 
-		TEST_METHOD(DelayedCommandExecuted_WhenAddedInBeforeEach_IsTrueDuringRunTest)
+		FTestTickable Tickable{};
+
+		TEST_METHOD(TestWithTickableGameObject_WaitingForTicksInSetup_WillAllowGameObjectToTick)
 		{
-			ASSERT_THAT(IsTrue(DelayedCommandExecuted));
+			TestCommandBuilder.Do([this]() { ASSERT_THAT(IsTrue(Tickable.TickCount > 2)); });
 		}
 
-		TEST_METHOD(DelayedCommandExecuted_WhenAddedDuringTest_IsTrueAfterTest)
+		TEST_METHOD(TestWithTickableGameObject_WaitingForTicksInSetup_WillBeCompleteDuringRunStep)
 		{
-			DelayedCommandExecuted = false;
-			AddCommand(new FQueueCommand(this, 3));
+			ASSERT_THAT(IsTrue(Tickable.TickCount > 2));
 		}
 	};
 } // namespace CQTestTests
diff --git a/Engine/Plugins/Tests/CQTest/Source/CQTestTests/Private/CQTestUnitTestHelper.cpp b/Engine/Plugins/Tests/CQTest/Source/CQTestTests/Private/CQTestUnitTestHelper.cpp
new file mode 100644
--- /dev/null
+++ b/Engine/Plugins/Tests/CQTest/Source/CQTestTests/Private/CQTestUnitTestHelper.cpp
@@ -0,0 +1,21 @@
+// Copyright Epic Games, Inc. All Rights Reserved.
+
+#pragma once
+#include "CQTestUnitTestHelper.h"
+
+void ClearExpectedError(FAutomationTestBase& TestRunner, const FString& ExpectedError)
+{
+	FAutomationTestExecutionInfo testInfo;
+	TestRunner.GetExecutionInfo(testInfo);
+	if (testInfo.GetErrorTotal() != 1)
+	{
+		return;
+	}
+	testInfo.RemoveAllEvents([&ExpectedError](FAutomationEvent& event) {
+		return event.Message.Equals(ExpectedError);
+		});
+	if (testInfo.GetErrorTotal() == 0)
+	{
+		TestRunner.ClearExecutionInfo();
+	}
+}
diff --git a/Engine/Plugins/Tests/CQTest/Source/CQTestTests/Private/CQTestUnitTestHelper.h b/Engine/Plugins/Tests/CQTest/Source/CQTestTests/Private/CQTestUnitTestHelper.h
new file mode 100644
--- /dev/null
+++ b/Engine/Plugins/Tests/CQTest/Source/CQTestTests/Private/CQTestUnitTestHelper.h
@@ -0,0 +1,9 @@
+// Copyright Epic Games, Inc. All Rights Reserved.
+
+#pragma once
+#include "Misc/AutomationTest.h"
+
+// When testing the framework, occasionally we need the error to live on the test framework
+// for short-circuiting logic.  We use this method to manually clear the error
+// (Normally from AFTER_EACH) to cause the test to pass
+void ClearExpectedError(FAutomationTestBase& TestRunner, const FString& ExpectedError);
diff --git a/Engine/Plugins/Tests/CQTest/Source/CQTestTests/Private/CommandBuilderTests.cpp b/Engine/Plugins/Tests/CQTest/Source/CQTestTests/Private/CommandBuilderTests.cpp
deleted file mode 100644
diff --git a/Engine/Plugins/Tests/CQTest/Source/CQTestTests/Private/Commands/CommandBuilderTests.cpp b/Engine/Plugins/Tests/CQTest/Source/CQTestTests/Private/Commands/CommandBuilderTests.cpp
new file mode 100644
--- /dev/null
+++ b/Engine/Plugins/Tests/CQTest/Source/CQTestTests/Private/Commands/CommandBuilderTests.cpp
@@ -0,0 +1,42 @@
+// Copyright Epic Games, Inc. All Rights Reserved.
+
+#include "CQTest.h"
+
+TEST_CLASS(CommandBuilderTests, "TestFramework.CQTest.Core")
+{
+	FTestCommandBuilder CommandBuilder{*TestRunner};
+
+	TEST_METHOD(Do_ThenBuild_IncludesCommand)
+	{
+		bool invoked = false;
+		auto command = CommandBuilder.Do([&invoked]() { invoked = true; }).Build();
+
+		ASSERT_THAT(IsTrue(command->Update()));
+		ASSERT_THAT(IsTrue(invoked));
+	}
+
+	TEST_METHOD(Build_WithoutCommands_ReturnsNullptr)
+	{
+		auto command = CommandBuilder.Build();
+		ASSERT_THAT(IsNull(command));
+	}
+
+	TEST_METHOD(StartWhen_CreatesWaitUntilCommand)
+	{
+		bool done = false;
+		auto command = CommandBuilder.StartWhen([&done]() { return done; }).Build();
+
+		ASSERT_THAT(IsFalse(command->Update()));
+		done = true;
+		ASSERT_THAT(IsTrue(command->Update()));
+	}
+
+	TEST_METHOD(Build_AfterBuild_ReturnsNullptr)
+	{
+		auto command = CommandBuilder.Do([]() {}).Build();
+		auto secondTime = CommandBuilder.Build();
+
+		ASSERT_THAT(IsNotNull(command));
+		ASSERT_THAT(IsNull(secondTime));
+	}
+};
\ No newline at end of file
diff --git a/Engine/Plugins/Tests/CQTest/Source/CQTestTests/Private/Commands/TestCommandsTests.cpp b/Engine/Plugins/Tests/CQTest/Source/CQTestTests/Private/Commands/TestCommandsTests.cpp
new file mode 100644
--- /dev/null
+++ b/Engine/Plugins/Tests/CQTest/Source/CQTestTests/Private/Commands/TestCommandsTests.cpp
@@ -0,0 +1,302 @@
+// Copyright Epic Games, Inc. All Rights Reserved.
+
+#include "CQTest.h"
+TEST_CLASS(RunSequenceBasicTests, "TestFramework.CQTest.Core")
+{
+	TSharedPtr<IAutomationLatentCommand> Cmd1;
+	TSharedPtr<IAutomationLatentCommand> Cmd2;
+	FString Cmd1Name = "One";
+	FString Cmd2Name = "Two";
+
+	TArray<FString> Log;
+
+	BEFORE_EACH() {
+		Cmd1 = MakeShared<FExecute>(*TestRunner, [&]() { Log.Add(Cmd1Name); }, *Cmd1Name);
+		Cmd2 = MakeShared<FExecute>(*TestRunner, [&]() { Log.Add(Cmd2Name); }, *Cmd2Name);
+	}
+
+	TEST_METHOD(Update_WithRemainingCommands_ReturnsFalse)
+	{
+		FRunSequence sequence(Cmd1, Cmd2);
+		ASSERT_THAT(IsFalse(sequence.Update()));
+	}
+
+	TEST_METHOD(Update_OnLastCommand_ReturnTrue)
+	{
+		FRunSequence sequence(Cmd1);
+		ASSERT_THAT(IsTrue(sequence.Update()));
+	}
+
+	TEST_METHOD(Append_ANewCommand_AddsCommandToEnd)
+	{
+		FRunSequence sequence{ Cmd1 };
+		sequence.Append(Cmd2);
+
+		sequence.Update();
+		sequence.Update();
+
+		ASSERT_THAT(AreEqual(2, Log.Num()));
+		ASSERT_THAT(AreEqual(Cmd1Name, Log[0]));
+		ASSERT_THAT(AreEqual(Cmd2Name, Log[1]));
+	}
+
+	TEST_METHOD(Prepend_ANewCommand_AddsCommandToBeginning)
+	{
+		FRunSequence sequence{ Cmd1 };
+		sequence.Prepend(Cmd2);
+
+		sequence.Update();
+		sequence.Update();
+
+		ASSERT_THAT(AreEqual(2, Log.Num()));
+		ASSERT_THAT(AreEqual(Cmd2Name, Log[0]));
+		ASSERT_THAT(AreEqual(Cmd1Name, Log[1]));
+	}
+};
+
+struct FCommandLog
+{
+	TArray<FString> Commands;
+};
+
+class FNamedCommand : public IAutomationLatentCommand
+{
+public:
+	FNamedCommand(TArray<FString>& CommandLog, FString Name)
+		: Log(CommandLog), CommandName(Name) {}
+
+	bool Update() override
+	{
+		Log.Add(CommandName);
+		return true;
+	}
+
+	TArray<FString>& Log;
+	FString CommandName;
+};
+
+class FTickingNamedCommand : public FNamedCommand
+{
+public:
+	FTickingNamedCommand(TArray<FString>& CommandLog, FString Name, int32 Ticks)
+		: FNamedCommand(CommandLog, Name), ExpectedCount(Ticks) {}
+
+	bool Update() override
+	{
+		if (CurrentCount == ExpectedCount)
+		{
+			return true;
+		}
+
+		Log.Add(CommandName);
+		CurrentCount++;
+		return false;
+	}
+
+	int32 ExpectedCount{ 0 };
+	int32 CurrentCount{ 0 };
+};
+
+
+TEST_CLASS(RunSequenceTests, "TestFramework.CQTest.Core")
+{
+	const TArray<FString> Names =
+	{
+		"Zero",
+		"One",
+		"Two",
+		"Three",
+		"Four",
+		"Five",
+		"Six",
+		"Seven",
+		"Eight"
+	};
+
+	TFunction<bool(RunSequenceTests*)> Assertion;
+	TArray<FString> CommandLog;
+
+	AFTER_EACH()
+	{
+		ASSERT_THAT(IsTrue(Assertion(this)));
+	}
+
+	TEST_METHOD(RunSequence_WithZeroCommands_DoesNotFail)
+	{
+		AddCommand(new FRunSequence());
+		Assertion = [](RunSequenceTests* test) {
+			return test->CommandLog.IsEmpty();
+		};
+	}
+
+	TEST_METHOD(RunSequence_WithOneCommand_RunsCommand)
+	{
+		AddCommand(new FRunSequence(MakeShared<FNamedCommand>(CommandLog, Names[0])));
+		Assertion = [](RunSequenceTests* test) {
+			return test->CommandLog.Num() == 1 && test->CommandLog[0] == test->Names[0];
+		};
+	}
+
+	TEST_METHOD(RunSequence_WithNamedCommands_RunsCommandsInOrder)
+	{
+		TArray<TSharedPtr<FNamedCommand>> Commands;
+		Commands.Add(MakeShared<FNamedCommand>(CommandLog, Names[0]));
+		Commands.Add(MakeShared<FNamedCommand>(CommandLog, Names[1]));
+		Commands.Add(MakeShared<FNamedCommand>(CommandLog, Names[2]));
+		Commands.Add(MakeShared<FNamedCommand>(CommandLog, Names[3]));
+		Commands.Add(MakeShared<FNamedCommand>(CommandLog, Names[4]));
+
+		AddCommand(new FRunSequence(Commands));
+
+		Assertion = [](RunSequenceTests* test) {
+			if (test->CommandLog.Num() != 5)
+			{
+				return false;
+			}
+
+			for (int32 i = 0; i < 5; i++)
+			{
+				if (test->CommandLog[i] != test->Names[i])
+				{
+					return false;
+				}
+			}
+
+			return true;
+		};
+	}
+
+	TEST_METHOD(RunSequence_WithTickingCommands_RunsCommandsInOrder)
+	{
+		TArray<TSharedPtr<FTickingNamedCommand>> Commands;
+		Commands.Add(MakeShared<FTickingNamedCommand>(CommandLog, Names[0], 3));
+		Commands.Add(MakeShared<FTickingNamedCommand>(CommandLog, Names[1], 3));
+		Commands.Add(MakeShared<FTickingNamedCommand>(CommandLog, Names[2], 3));
+		Commands.Add(MakeShared<FTickingNamedCommand>(CommandLog, Names[3], 3));
+		Commands.Add(MakeShared<FTickingNamedCommand>(CommandLog, Names[4], 3));
+
+		AddCommand(new FRunSequence(Commands));
+
+		Assertion = [](RunSequenceTests* test) {
+			if (test->CommandLog.Num() != 15)
+			{
+				return false;
+			}
+
+			int32 NameIndex = -1;
+			for (int32 CommandIndex = 0; CommandIndex < test->CommandLog.Num(); CommandIndex++)
+			{
+				if (CommandIndex % 3 == 0)
+				{
+					NameIndex++;
+				}
+
+				if (test->CommandLog[CommandIndex] != test->Names[NameIndex])
+				{
+					return false;
+				}
+			}
+
+			return true;
+		};
+	}
+
+	TEST_METHOD(RunSequence_WithSequences_RunsCommandsInOrder)
+	{
+		TArray<TSharedPtr<FNamedCommand>> Cmds1;
+		TArray<TSharedPtr<FNamedCommand>> Cmds2;
+		TArray<TSharedPtr<FNamedCommand>> Cmds3;
+
+		Cmds1.Add(MakeShared<FNamedCommand>(CommandLog, Names[0]));
+		Cmds1.Add(MakeShared<FNamedCommand>(CommandLog, Names[1]));
+		Cmds1.Add(MakeShared<FNamedCommand>(CommandLog, Names[2]));
+
+		Cmds2.Add(MakeShared<FNamedCommand>(CommandLog, Names[3]));
+		Cmds2.Add(MakeShared<FNamedCommand>(CommandLog, Names[4]));
+		Cmds2.Add(MakeShared<FNamedCommand>(CommandLog, Names[5]));
+
+		Cmds3.Add(MakeShared<FNamedCommand>(CommandLog, Names[6]));
+		Cmds3.Add(MakeShared<FNamedCommand>(CommandLog, Names[7]));
+		Cmds3.Add(MakeShared<FNamedCommand>(CommandLog, Names[8]));
+
+		AddCommand(new FRunSequence(MakeShared<FRunSequence>(Cmds1), MakeShared<FRunSequence>(Cmds2), MakeShared<FRunSequence>(Cmds3)));
+
+		Assertion = [](RunSequenceTests* test) {
+			for (int32 i = 0; i < 9; i++)
+			{
+				if (test->CommandLog[i] != test->Names[i])
+				{
+					return false;
+				}
+			}
+
+			return true;
+		};
+	}
+
+	TEST_METHOD(RunSequence_WithSeparateSequences_RunsCommandsInOrder)
+	{
+		TArray<TSharedPtr<FNamedCommand>> Cmds1;
+		TArray<TSharedPtr<FNamedCommand>> Cmds2;
+		TArray<TSharedPtr<FNamedCommand>> Cmds3;
+
+		Cmds1.Add(MakeShared<FNamedCommand>(CommandLog, Names[0]));
+		Cmds1.Add(MakeShared<FNamedCommand>(CommandLog, Names[1]));
+		Cmds1.Add(MakeShared<FNamedCommand>(CommandLog, Names[2]));
+
+		Cmds2.Add(MakeShared<FNamedCommand>(CommandLog, Names[3]));
+		Cmds2.Add(MakeShared<FNamedCommand>(CommandLog, Names[4]));
+		Cmds2.Add(MakeShared<FNamedCommand>(CommandLog, Names[5]));
+
+		Cmds3.Add(MakeShared<FNamedCommand>(CommandLog, Names[6]));
+		Cmds3.Add(MakeShared<FNamedCommand>(CommandLog, Names[7]));
+		Cmds3.Add(MakeShared<FNamedCommand>(CommandLog, Names[8]));
+
+		AddCommand(new FRunSequence(Cmds1));
+		AddCommand(new FRunSequence(Cmds2));
+		AddCommand(new FRunSequence(Cmds3));
+
+		Assertion = [](RunSequenceTests* test) {
+			for (int32 i = 0; i < 9; i++)
+			{
+				if (test->CommandLog[i] != test->Names[i])
+				{
+					return false;
+				}
+			}
+
+			return true;
+		};
+	}
+
+	TEST_METHOD(RunSequence_WithUntilCommands_RunsCommandsInOrder)
+	{
+		TArray<TSharedPtr<IAutomationLatentCommand>> Cmds;
+		Cmds.Add(MakeShared<FWaitUntil>(*TestRunner, [&]() {
+			static int32 attempt = 0;
+			CommandLog.Add(Names[0]);
+			if (++attempt > 3)
+			{
+				attempt = 0;
+				return true;
+			}
+			return false;
+			}));
+		Cmds.Add(MakeShared<FWaitUntil>(*TestRunner, [&]() {
+			static int32 attempt = 0;
+			CommandLog.Add(Names[1]);
+			if (++attempt > 4)
+			{
+				attempt = 0;
+				return true;
+			}
+			return false;
+			}));
+
+		AddCommand(new FRunSequence(Cmds));
+
+		Assertion = [](RunSequenceTests* test) {
+			return test->CommandLog.Num() == 9;
+		};
+	}
+};
\ No newline at end of file
diff --git a/Engine/Plugins/Tests/CQTest/Source/CQTestTests/Private/Components/PIENetworkComponentTests.cpp b/Engine/Plugins/Tests/CQTest/Source/CQTestTests/Private/Components/PIENetworkComponentTests.cpp
new file mode 100644
--- /dev/null
+++ b/Engine/Plugins/Tests/CQTest/Source/CQTestTests/Private/Components/PIENetworkComponentTests.cpp
@@ -0,0 +1,188 @@
+// Copyright Epic Games, Inc. All Rights Reserved.
+
+#include "CQTest.h"
+#include "Components/PIENetworkComponent.h"
+#include "CQTestUnitTestHelper.h"
+
+#include "TestReplicatedActor.h"
+
+#include "Kismet/GameplayStatics.h"
+#include "GameFramework/GameModeBase.h"
+#include "Engine/GameInstance.h"
+#include "Engine/World.h"
+
+#if ENABLE_PIE_NETWORK_TEST
+
+NETWORK_TEST_CLASS(StateTest, "TestFramework.CQTest.Network")
+{
+	struct DerivedState : public FBasePIENetworkComponentState
+	{
+		int32 IndependentNumber = 0;
+	};
+
+	int32 SharedNumber = 0;
+	FPIENetworkComponent<DerivedState> Network{ TestRunner, TestCommandBuilder, bInitializing };
+
+	BEFORE_EACH()
+	{
+		FNetworkComponentBuilder<DerivedState>()
+			.WithClients(3)
+			.WithGameInstanceClass(UGameInstance::StaticClass())
+			.WithGameMode(AGameModeBase::StaticClass())
+			.Build(Network);
+	}
+
+	TEST_METHOD(Network_WithMultipleSteps_TriggersStepsInOrder)
+	{
+		Network
+			.ThenServer([&](DerivedState&) { ASSERT_THAT(AreEqual(0, SharedNumber++)); })
+			.ThenClient(0, [&](DerivedState&) { ASSERT_THAT(AreEqual(1, SharedNumber++)); })
+			.ThenServer([&](DerivedState&) { ASSERT_THAT(AreEqual(2, SharedNumber++)); })
+			.Then([&]() { ASSERT_THAT(AreEqual(3, SharedNumber++)); });
+	}
+
+	TEST_METHOD(Network_WithServerCommands_RetainsStateBetweenCalls)
+	{
+		Network
+			.ThenServer([&](DerivedState& State) { ASSERT_THAT(AreEqual(0, State.IndependentNumber++)); })
+			.ThenServer([&](DerivedState& State) { ASSERT_THAT(AreEqual(1, State.IndependentNumber++)); });
+	}
+
+	TEST_METHOD(Network_WithClientCommands_RetainsStateBetweenCalls)
+	{
+		Network
+			.ThenClient(0, [&](DerivedState& State) { ASSERT_THAT(AreEqual(0, State.IndependentNumber++)); })
+			.ThenClient(0, [&](DerivedState& State) { ASSERT_THAT(AreEqual(1, State.IndependentNumber++)); });
+	}
+
+	TEST_METHOD(Network_WithClientAndServerCommands_DoNotShareState)
+	{
+		Network
+			.ThenServer([&](DerivedState& State) { State.IndependentNumber++; })
+			.ThenClient(0, [&](DerivedState& State) { ASSERT_THAT(AreEqual(0, State.IndependentNumber)); });
+	}
+
+	TEST_METHOD(Network_WithMultipleClients_DoNotShareState)
+	{
+		Network
+			.ThenClients([&](DerivedState& State) { State.IndependentNumber = State.ClientIndex; })
+			.ThenClients([&](DerivedState& State) { ASSERT_THAT(AreEqual(State.ClientIndex, State.IndependentNumber)); });
+	}
+
+	TEST_METHOD(Network_WithTickingServerCommand_TicksUntilDone)
+	{
+		Network
+			.UntilServer([&](DerivedState& State) { return ++State.IndependentNumber > 4; })
+			.ThenServer([&](DerivedState& State) { ASSERT_THAT(AreEqual(State.IndependentNumber, 5)); });
+	}
+
+	TEST_METHOD(Network_WithTickingClientCommands_TicksEachCommand)
+	{
+		Network
+			.UntilClients([&](DerivedState& State) { SharedNumber++;  return ++State.IndependentNumber > 4; })
+			.Then([&]() { ASSERT_THAT(AreEqual(15, SharedNumber)); });
+	}
+};
+
+NETWORK_TEST_CLASS(ReplicationTest, "TestFramework.CQTest.Network")
+{
+	struct DerivedState : public FBasePIENetworkComponentState
+	{
+		ATestReplicatedActor* ReplicatedActor = nullptr;
+	};
+
+	FPIENetworkComponent<DerivedState> Network{ TestRunner, TestCommandBuilder, bInitializing };
+	BEFORE_EACH() {
+		FNetworkComponentBuilder<DerivedState>()
+			.WithGameInstanceClass(UGameInstance::StaticClass())
+			.WithGameMode(AGameModeBase::StaticClass())
+			.Build(Network);
+	}
+
+	TEST_METHOD(SpawnActor_WithReplicatedActor_ProvidesActorToClients)
+	{
+		Network.SpawnAndReplicate<ATestReplicatedActor, &DerivedState::ReplicatedActor>()
+			.ThenServer([&](DerivedState& ServerState) {
+				ASSERT_THAT(IsNotNull(ServerState.ReplicatedActor));
+			})
+			.ThenClients([&](DerivedState& ClientState) {
+				ASSERT_THAT(IsNotNull(ClientState.ReplicatedActor));
+			});
+	}
+
+	TEST_METHOD(SpawnAndReplicateActor_ThenUpdateProperty_UpdatesPropertyOnClients)
+	{
+		Network.SpawnAndReplicate<ATestReplicatedActor, &DerivedState::ReplicatedActor>()
+			.ThenServer(TEXT("Server Set Value"), [](DerivedState& ServerState) {
+				ServerState.ReplicatedActor->ReplicatedInt = 42;
+			})
+			.UntilClients(TEXT("Clients Check Value"), [](DerivedState& ClientState) {
+				return ClientState.ReplicatedActor->ReplicatedInt == 42;
+			});
+	}
+};
+
+NETWORK_TEST_CLASS(GameInstanceTest, "TestFramework.CQTest.Network")
+{
+	struct SomeGameInstanceClass : UGameInstance
+	{
+	};
+
+	FPIENetworkComponent<FBasePIENetworkComponentState> Network{ TestRunner, TestCommandBuilder, bInitializing };
+	BEFORE_EACH()
+	{
+		FNetworkComponentBuilder()
+			.WithGameInstanceClass(FSoftClassPath(SomeGameInstanceClass::StaticClass()))
+			.WithGameMode(AGameModeBase::StaticClass())
+			.Build(Network);
+	}
+
+	TEST_METHOD(NetworkComponent_WithGameInstanceClass_BuildsNetworkWithProvidedGameInstance)
+	{
+		Network.ThenServer([&](FBasePIENetworkComponentState& ServerState) {
+			auto* GameInstance = Cast<SomeGameInstanceClass>(ServerState.World->GetGameInstance());
+			ASSERT_THAT(IsNotNull(GameInstance));
+		});
+	}
+};
+
+NETWORK_TEST_CLASS(GameModeTest, "TestFramework.CQTest.Network")
+{
+	struct ATestGameMode : AGameModeBase
+	{
+	};
+
+	FPIENetworkComponent<> Network{ TestRunner, TestCommandBuilder, bInitializing };
+
+	BEFORE_EACH()
+	{
+		FNetworkComponentBuilder()
+			.WithGameInstanceClass(UGameInstance::StaticClass())
+			.WithGameMode(ATestGameMode::StaticClass())
+			.Build(Network);
+	}
+
+	TEST_METHOD(NetworkComponent_WithGameMode_BuildsNetworkWithProvidedGameMode)
+	{
+		Network.ThenServer([&](auto& ServerState) {
+			auto* GameMode = Cast<ATestGameMode>(ServerState.World->GetAuthGameMode());
+			ASSERT_THAT(IsNotNull(GameMode));
+		});
+	}
+};
+
+NETWORK_TEST_CLASS(SetupErrorTest, "TestFramework.CQTest.Network")
+{
+	FPIENetworkComponent<> Network{ TestRunner, TestCommandBuilder, bInitializing };
+
+	AFTER_EACH() 
+	{
+		ClearExpectedError(*TestRunner, TEXT("Failed to initialize Network Component"));
+	}
+	TEST_METHOD(NetworkComponent_WithoutUsingBuilder_AddsErrorAndDoesNotCrash)
+	{
+		Network.ThenServer([&](auto&) { Assert.Fail("Unexpected Error"); });
+	}
+};
+
+#endif // ENABLE_PIE_NETWORK_TEST
\ No newline at end of file
diff --git a/Engine/Plugins/Tests/CQTest/Source/CQTestTests/Private/Components/SpawnHelperTests.cpp b/Engine/Plugins/Tests/CQTest/Source/CQTestTests/Private/Components/SpawnHelperTests.cpp
new file mode 100644
--- /dev/null
+++ b/Engine/Plugins/Tests/CQTest/Source/CQTestTests/Private/Components/SpawnHelperTests.cpp
@@ -0,0 +1,30 @@
+// Copyright Epic Games, Inc. All Rights Reserved.
+
+#include "CQTest.h"
+#include "Components/ActorTestSpawner.h"
+
+TEST_CLASS(SpawnHelperTests, "TestFramework.CQTest.Actor")
+{
+	FActorTestSpawner Spawner;
+
+	TEST_METHOD(SpawnActor_WithNameParameter_ReturnsActorWithGivenName)
+	{
+		const FName ExpectedActorName = TEXT("AnyActorName");
+		FActorSpawnParameters SpawnParameters;
+		SpawnParameters.Name = ExpectedActorName;
+
+		AActor& Actor = Spawner.SpawnActor<AActor>(SpawnParameters);
+
+		ASSERT_THAT(AreEqual( Actor.GetFName(), ExpectedActorName ));
+	}
+	
+	TEST_METHOD(AddComponentTo_SpawnedActor_SetsActorAsComponentOwner)
+	{
+		AActor& Owner = Spawner.SpawnActor<AActor>();
+		auto* TestComponent = NewObject<USceneComponent>(&Owner);
+		TestComponent->RegisterComponent();
+
+		ASSERT_THAT(IsNotNull(TestComponent));
+		ASSERT_THAT(AreEqual(TestComponent->GetOwner(), &Owner ));
+	}
+};
\ No newline at end of file
diff --git a/Engine/Plugins/Tests/CQTest/Source/CQTestTests/Private/Components/TestReplicatedActor.cpp b/Engine/Plugins/Tests/CQTest/Source/CQTestTests/Private/Components/TestReplicatedActor.cpp
new file mode 100644
--- /dev/null
+++ b/Engine/Plugins/Tests/CQTest/Source/CQTestTests/Private/Components/TestReplicatedActor.cpp
@@ -0,0 +1,29 @@
+// Copyright Epic Games, Inc. All Rights Reserved.
+
+#pragma once
+
+#include "Components/TestReplicatedActor.h"
+
+#include "Net/UnrealNetwork.h"
+
+// -----------------------------------------------------------------------------
+
+ATestReplicatedActor::ATestReplicatedActor(const FObjectInitializer& ObjectInitializer)
+	: Super(ObjectInitializer)
+	, ReplicatedInt(0)
+{
+	bReplicates = true;
+	bAlwaysRelevant = true;
+}
+
+// -----------------------------------------------------------------------------
+
+void ATestReplicatedActor::GetLifetimeReplicatedProps(TArray<FLifetimeProperty>& OutLifetimeProps) const
+{
+	Super::GetLifetimeReplicatedProps(OutLifetimeProps);
+	DOREPLIFETIME(ATestReplicatedActor, ReplicatedInt);
+}
+
+// -----------------------------------------------------------------------------
+// -----------------------------------------------------------------------------
+// -----------------------------------------------------------------------------
diff --git a/Engine/Plugins/Tests/CQTest/Source/CQTestTests/Private/Components/TestReplicatedActor.h b/Engine/Plugins/Tests/CQTest/Source/CQTestTests/Private/Components/TestReplicatedActor.h
new file mode 100644
--- /dev/null
+++ b/Engine/Plugins/Tests/CQTest/Source/CQTestTests/Private/Components/TestReplicatedActor.h
@@ -0,0 +1,27 @@
+// Copyright Epic Games, Inc. All Rights Reserved.
+
+#pragma once
+
+#include "GameFramework/Actor.h"
+
+#include "TestReplicatedActor.generated.h"
+
+// -----------------------------------------------------------------------------
+
+UCLASS(NotBlueprintable)
+class ATestReplicatedActor : public AActor
+{
+	GENERATED_BODY()
+
+public:
+	ATestReplicatedActor(const FObjectInitializer& ObjectInitializer = FObjectInitializer::Get());
+
+	void GetLifetimeReplicatedProps(TArray<FLifetimeProperty>& OutLifetimeProps) const override;
+
+	UPROPERTY(Replicated)
+	int32 ReplicatedInt;
+};
+
+// -----------------------------------------------------------------------------
+// -----------------------------------------------------------------------------
+// -----------------------------------------------------------------------------
diff --git a/Engine/Plugins/Tests/CQTest/Source/CQTestTests/Private/ExtensionTests.cpp b/Engine/Plugins/Tests/CQTest/Source/CQTestTests/Private/ExtensionTests.cpp
--- a/Engine/Plugins/Tests/CQTest/Source/CQTestTests/Private/ExtensionTests.cpp
+++ b/Engine/Plugins/Tests/CQTest/Source/CQTestTests/Private/ExtensionTests.cpp
@@ -11,17 +11,19 @@ public:
 	{
 	}
 
-	void Custom(bool) {
+	bool Custom(bool In) 
+	{
+		return In;
 	}
 };
 
 #define CUSTOM_ASSERT_TEST_CLASS(_ClassName, _TestDir) TEST_CLASS_WITH_ASSERTS(_ClassName, _TestDir, FCustomAsserter)
 
-CUSTOM_ASSERT_TEST_CLASS(CustomAsserts, "TestFramework.CQTest")
+CUSTOM_ASSERT_TEST_CLASS(CustomAsserts, "TestFramework.CQTest.Core")
 {
 	TEST_METHOD(CustomTestClass_WithCustomAsserter_HasInstanceOfCustomAsserter)
 	{
-		Assert.Custom(true);
+		ASSERT_THAT(Custom(true));
 		ASSERT_THAT(IsTrue(true));
 	}
 };
@@ -34,7 +36,7 @@ struct TCustomBaseClass : public TTest<Derived, AsserterType>
 
 #define CUSTOM_BASE_TEST_CLASS(_ClassName, _TestDir) TEST_CLASS_WITH_BASE(_ClassName, _TestDir, TCustomBaseClass)
 
-CUSTOM_BASE_TEST_CLASS(CustomBase, "TestFramework.CQTest")
+CUSTOM_BASE_TEST_CLASS(CustomBase, "TestFramework.CQTest.Core")
 {
 	TEST_METHOD(CustomTestClass_WithCustomBase_InheritsFromBaseClass)
 	{
@@ -64,7 +66,7 @@ struct TBaseWithConstructor : public TTest<Derived, AsserterType>
 
 #define CUSTOM_WITH_CTOR_CLASS(_ClassName, _TestDir) TEST_CLASS_WITH_BASE(_ClassName, _TestDir, TBaseWithConstructor)
 
-CUSTOM_WITH_CTOR_CLASS(CustomCtor, "TestFramework.CQTest")
+CUSTOM_WITH_CTOR_CLASS(CustomCtor, "TestFramework.CQTest.Core")
 {
 	TEST_METHOD(CustomClassWithCtor_CallsCtorWhenInitializing_AndBeforeEachTest)
 	{
diff --git a/Engine/Plugins/Tests/CQTest/Source/CQTestTests/Private/ObjectBuilderTests.cpp b/Engine/Plugins/Tests/CQTest/Source/CQTestTests/Private/ObjectBuilderTests.cpp
--- a/Engine/Plugins/Tests/CQTest/Source/CQTestTests/Private/ObjectBuilderTests.cpp
+++ b/Engine/Plugins/Tests/CQTest/Source/CQTestTests/Private/ObjectBuilderTests.cpp
@@ -6,12 +6,12 @@
 #include "Components/SceneComponent.h"
 
 #include "ObjectBuilder.h"
-#include "ActorTestSpawner.h"
+#include "Components/ActorTestSpawner.h"
 #include "CQTest.h"
 
 #if WITH_AUTOMATION_WORKER
 
-TEST_CLASS(ObjectBuilder_Success, "TestFramework.CQTest")
+TEST_CLASS(ObjectBuilder_Success, "TestFramework.CQTest.Actor")
 {
 	FActorTestSpawner Spawner;
 
@@ -449,7 +449,7 @@ TEST_CLASS(ObjectBuilder_Success, "TestFramework.CQTest")
 	}
 };
 
-TEST_CLASS(ObjectBuilder_ApiError, "TestFramework.CQTest")
+TEST_CLASS(ObjectBuilder_ApiError, "TestFramework.CQTest.Actor")
 {
 	FActorTestSpawner Spawner;
 
@@ -482,7 +482,7 @@ TEST_CLASS(ObjectBuilder_ApiError, "TestFramework.CQTest")
 	}
 };
 
-TEST_CLASS(ObjectBuilder_PrimitiveParameterErrors, "TestFramework.CQTest")
+TEST_CLASS(ObjectBuilder_PrimitiveParameterErrors, "TestFramework.CQTest.Actor")
 {
 	FActorTestSpawner Spawner;
 
@@ -621,7 +621,7 @@ TEST_CLASS(ObjectBuilder_PrimitiveParameterErrors, "TestFramework.CQTest")
 	}
 };
 
-TEST_CLASS(ObjectBuilder_StructParameterErrors, "TestFramework.CQTest")
+TEST_CLASS(ObjectBuilder_StructParameterErrors, "TestFramework.CQTest.Actor")
 {
 	FActorTestSpawner Spawner;
 
@@ -667,7 +667,7 @@ TEST_CLASS(ObjectBuilder_StructParameterErrors, "TestFramework.CQTest")
 };
 
 
-TEST_CLASS(ObjectBuilder_ArrayPropertyErrors, "TestFramework.CQTest")
+TEST_CLASS(ObjectBuilder_ArrayPropertyErrors, "TestFramework.CQTest.Actor")
 {
 	FActorTestSpawner Spawner;
 
@@ -810,7 +810,7 @@ TEST_CLASS(ObjectBuilder_ArrayPropertyErrors, "TestFramework.CQTest")
 	}
 };
 
-TEST_CLASS(ObjectBuilder_SetPropertyErrors, "TestFramework.CQTest")
+TEST_CLASS(ObjectBuilder_SetPropertyErrors, "TestFramework.CQTest.Actor")
 {
 	FActorTestSpawner Spawner;
 
@@ -943,7 +943,7 @@ TEST_CLASS(ObjectBuilder_SetPropertyErrors, "TestFramework.CQTest")
 	}
 };
 
-TEST_CLASS(ObjectBuilder_MapPropertyErrors, "TestFramework.CQTest")
+TEST_CLASS(ObjectBuilder_MapPropertyErrors, "TestFramework.CQTest.Actor")
 {
 	FActorTestSpawner Spawner;
 
@@ -1077,7 +1077,7 @@ TEST_CLASS(ObjectBuilder_MapPropertyErrors, "TestFramework.CQTest")
 	}
 };
 
-TEST_CLASS(ObjectBuilder_PostSpawnErrors, "TestFramework.CQTest")
+TEST_CLASS(ObjectBuilder_PostSpawnErrors, "TestFramework.CQTest.Actor")
 {
 	FActorTestSpawner Spawner;
 
diff --git a/Engine/Plugins/Tests/CQTest/Source/CQTestTests/Private/SpawnHelperTests.cpp b/Engine/Plugins/Tests/CQTest/Source/CQTestTests/Private/SpawnHelperTests.cpp
deleted file mode 100644
diff --git a/Engine/Plugins/Tests/CQTest/Source/CQTestTests/Private/TestCommandsTests.cpp b/Engine/Plugins/Tests/CQTest/Source/CQTestTests/Private/TestCommandsTests.cpp
deleted file mode 100644
