Actually used InGame parameter in PIENetworkTestStateRestorer

#rnx

diff --git a/Engine/Plugins/Tests/CQTest/Source/CQTest/Private/Components/PIENetworkTestStateRestorer.cpp b/Engine/Plugins/Tests/CQTest/Source/CQTest/Private/Components/PIENetworkTestStateRestorer.cpp
--- a/Engine/Plugins/Tests/CQTest/Source/CQTest/Private/Components/PIENetworkTestStateRestorer.cpp
+++ b/Engine/Plugins/Tests/CQTest/Source/CQTest/Private/Components/PIENetworkTestStateRestorer.cpp
@@ -24,7 +24,11 @@ FPIENetworkTestStateRestorer::FPIENetworkTestStateRestorer(const FSoftClassPath
 			UE_LOG(LogPieNetworkStateRestorer, Error, TEXT("Unable to get UGameMapsSettings in when creating FPIENetworkTestStateRestorer"));
 		}
 	}
-	OriginalGameMode = EditorWorld->GetWorldSettings()->DefaultGameMode;
+	if (EditorWorld->GetWorldSettings() != nullptr)
+	{
+		OriginalGameMode = EditorWorld->GetWorldSettings()->DefaultGameMode;
+		EditorWorld->GetWorldSettings()->DefaultGameMode = InGameMode;
+	}
 	if (!EditorWorld->HasAllFlags(RF_WasLoaded))
 	{
 		EditorWorld->SetFlags(RF_WasLoaded);
