Change a CQ test flag from Smoke test to Engine test

#rnx

diff --git a/Engine/Plugins/Tests/CQTest/Source/CQTestTests/Private/CQTestTests.cpp b/Engine/Plugins/Tests/CQTest/Source/CQTestTests/Private/CQTestTests.cpp
--- a/Engine/Plugins/Tests/CQTest/Source/CQTestTests/Private/CQTestTests.cpp
+++ b/Engine/Plugins/Tests/CQTest/Source/CQTestTests/Private/CQTestTests.cpp
@@ -68,11 +68,11 @@ namespace CQTestTests
 		}
 	};
 
-	TEST_CLASS_WITH_FLAGS(OverrideFixtureTestFlags, "TestFramework.CQTest.Core", EAutomationTestFlags::ApplicationContextMask | EAutomationTestFlags::SmokeFilter)
+	TEST_CLASS_WITH_FLAGS(OverrideFixtureTestFlags, "TestFramework.CQTest.Core", EAutomationTestFlags::ApplicationContextMask | EAutomationTestFlags::EngineFilter)
 	{
 		TEST_METHOD(GetTestFlags_ReturnsSetAutomationTestFlags)
 		{
-			ASSERT_THAT(AreEqual(EAutomationTestFlags::ApplicationContextMask | EAutomationTestFlags::SmokeFilter, TestRunner->GetTestFlags()));
+			ASSERT_THAT(AreEqual(EAutomationTestFlags::ApplicationContextMask | EAutomationTestFlags::EngineFilter, TestRunner->GetTestFlags()));
 		}
 	};
 
