Fix NonUnity Compile UnrealGame IOS in CQTestUnitTestHelper.cpp and TestReplicatedActor.cpp

#rnx

diff --git a/Engine/Plugins/Tests/CQTest/Source/CQTestTests/Private/CQTestUnitTestHelper.cpp b/Engine/Plugins/Tests/CQTest/Source/CQTestTests/Private/CQTestUnitTestHelper.cpp
--- a/Engine/Plugins/Tests/CQTest/Source/CQTestTests/Private/CQTestUnitTestHelper.cpp
+++ b/Engine/Plugins/Tests/CQTest/Source/CQTestTests/Private/CQTestUnitTestHelper.cpp
@@ -1,6 +1,5 @@
 // Copyright Epic Games, Inc. All Rights Reserved.
 
-#pragma once
 #include "CQTestUnitTestHelper.h"
 
 void ClearExpectedError(FAutomationTestBase& TestRunner, const FString& ExpectedError)
diff --git a/Engine/Plugins/Tests/CQTest/Source/CQTestTests/Private/Components/TestReplicatedActor.cpp b/Engine/Plugins/Tests/CQTest/Source/CQTestTests/Private/Components/TestReplicatedActor.cpp
--- a/Engine/Plugins/Tests/CQTest/Source/CQTestTests/Private/Components/TestReplicatedActor.cpp
+++ b/Engine/Plugins/Tests/CQTest/Source/CQTestTests/Private/Components/TestReplicatedActor.cpp
@@ -1,7 +1,5 @@
 // Copyright Epic Games, Inc. All Rights Reserved.
 
-#pragma once
-
 #include "Components/TestReplicatedActor.h"
 
 #include "Net/UnrealNetwork.h"
