Fix ASan error heap-use-after-free in FPIENetworkTestStateRestorer
Delay Garbage collection when a test start and Force it at the end of a test
Add a Teardown queue to test command builder

#rnx

diff --git a/Engine/Plugins/Tests/CQTest/Source/CQTest/Private/Commands/TestCommands.cpp b/Engine/Plugins/Tests/CQTest/Source/CQTest/Private/Commands/TestCommands.cpp
--- a/Engine/Plugins/Tests/CQTest/Source/CQTest/Private/Commands/TestCommands.cpp
+++ b/Engine/Plugins/Tests/CQTest/Source/CQTest/Private/Commands/TestCommands.cpp
@@ -42,11 +42,12 @@ bool CQTEST_API FWaitUntil::Update()
 
 bool CQTEST_API FExecute::Update()
 {
+	bool bHasToRun = !TestRunner.HasAnyErrors() || FailureBehavior == ECQTestFailureBehavior::Run;
 	if (Description)
 	{
-		UE_LOG(LogCqTest, Log, TEXT("Running %s"), Description);
+		UE_LOG(LogCqTest, Log, TEXT("%s %s"), bHasToRun? TEXT("Running") : TEXT("Skipping"), Description);
 	}
-	if (!TestRunner.HasAnyErrors())
+	if (bHasToRun)
 	{
 		Func();
 	}
diff --git a/Engine/Plugins/Tests/CQTest/Source/CQTest/Private/Components/PIENetworkComponent.cpp b/Engine/Plugins/Tests/CQTest/Source/CQTest/Private/Components/PIENetworkComponent.cpp
--- a/Engine/Plugins/Tests/CQTest/Source/CQTest/Private/Components/PIENetworkComponent.cpp
+++ b/Engine/Plugins/Tests/CQTest/Source/CQTest/Private/Components/PIENetworkComponent.cpp
@@ -29,16 +29,8 @@ FBasePIENetworkComponent::FBasePIENetworkComponent(FAutomationTestBase* InTestRu
 	CommandBuilder
 		->Do(TEXT("Start PIE"), [this]() { StartPie(); })
 		.Until(TEXT("Collect PIE Worlds"), [this]() { return CollectPieWorlds(); })
-		.Until(TEXT("Await Connections"), [this]() { return AwaitConnections(); });
-}
-
-FBasePIENetworkComponent::~FBasePIENetworkComponent()
-{
-	if (ServerState != nullptr)
-	{
-		GUnrealEd->RequestEndPlayMap();
-		StateRestorer.Restore();
-	}
+		.Until(TEXT("Await Connections"), [this]() { return AwaitConnections(); })
+		.OnTearDown(TEXT("Restore Editor State"), [this]() { RestoreState(); });
 }
 
 FBasePIENetworkComponent& FBasePIENetworkComponent::Then(TFunction<void()> Action)
@@ -207,4 +199,13 @@ bool FBasePIENetworkComponent::AwaitConnections()
 
 	return true;
 }
+
+void FBasePIENetworkComponent::RestoreState()
+{
+	if (ServerState != nullptr)
+	{
+		GUnrealEd->RequestEndPlayMap();
+		StateRestorer.Restore();
+	}
+}
 #endif
\ No newline at end of file
diff --git a/Engine/Plugins/Tests/CQTest/Source/CQTest/Private/Components/PIENetworkTestStateRestorer.cpp b/Engine/Plugins/Tests/CQTest/Source/CQTest/Private/Components/PIENetworkTestStateRestorer.cpp
--- a/Engine/Plugins/Tests/CQTest/Source/CQTest/Private/Components/PIENetworkTestStateRestorer.cpp
+++ b/Engine/Plugins/Tests/CQTest/Source/CQTest/Private/Components/PIENetworkTestStateRestorer.cpp
@@ -10,7 +10,7 @@ DEFINE_LOG_CATEGORY_STATIC(LogPieNetworkStateRestorer, Log, All);
 
 FPIENetworkTestStateRestorer::FPIENetworkTestStateRestorer(const FSoftClassPath InGameInstanceClass, const TSubclassOf<AGameModeBase> InGameMode)
 {
-	EditorWorld = GWorld;
+	check(GWorld);
 	if (InGameInstanceClass.IsValid())
 	{
 		UGameMapsSettings* GameMapSettings = GetMutableDefault<UGameMapsSettings>();
@@ -24,22 +24,22 @@ FPIENetworkTestStateRestorer::FPIENetworkTestStateRestorer(const FSoftClassPath
 			UE_LOG(LogPieNetworkStateRestorer, Error, TEXT("Unable to get UGameMapsSettings in when creating FPIENetworkTestStateRestorer"));
 		}
 	}
-	if (EditorWorld->GetWorldSettings() != nullptr)
+	if (GWorld->GetWorldSettings() != nullptr)
 	{
-		OriginalGameMode = EditorWorld->GetWorldSettings()->DefaultGameMode;
-		EditorWorld->GetWorldSettings()->DefaultGameMode = InGameMode;
+		OriginalGameMode = GWorld->GetWorldSettings()->DefaultGameMode;
+		GWorld->GetWorldSettings()->DefaultGameMode = InGameMode;
 	}
-	if (!EditorWorld->HasAllFlags(RF_WasLoaded))
+	if (!GWorld->HasAllFlags(RF_WasLoaded))
 	{
-		EditorWorld->SetFlags(RF_WasLoaded);
+		GWorld->SetFlags(RF_WasLoaded);
 		SetWasLoadedFlag = true;
 	}
 }
 
 void FPIENetworkTestStateRestorer::Restore() {
-	if (SetWasLoadedFlag)
+	if (GWorld && SetWasLoadedFlag)
 	{
-		EditorWorld->ClearFlags(RF_WasLoaded);
+		GWorld->ClearFlags(RF_WasLoaded);
 	}
 
 	if (OriginalGameInstance.IsValid())
@@ -55,8 +55,8 @@ void FPIENetworkTestStateRestorer::Restore() {
 		}
 	}
 
-	if (OriginalGameMode != nullptr)
+	if (GWorld && OriginalGameMode != nullptr)
 	{
-		EditorWorld->GetWorldSettings()->DefaultGameMode = OriginalGameMode;
+		GWorld->GetWorldSettings()->DefaultGameMode = OriginalGameMode;
 	}
 }
\ No newline at end of file
diff --git a/Engine/Plugins/Tests/CQTest/Source/CQTest/Public/CQTest.h b/Engine/Plugins/Tests/CQTest/Source/CQTest/Public/CQTest.h
--- a/Engine/Plugins/Tests/CQTest/Source/CQTest/Public/CQTest.h
+++ b/Engine/Plugins/Tests/CQTest/Source/CQTest/Public/CQTest.h
@@ -2,6 +2,8 @@
 
 #pragma once
 
+#include "Engine/Engine.h"
+
 #include "Misc/AutomationTest.h"
 #include "HAL/Platform.h"
 
diff --git a/Engine/Plugins/Tests/CQTest/Source/CQTest/Public/Commands/TestCommandBuilder.h b/Engine/Plugins/Tests/CQTest/Source/CQTest/Public/Commands/TestCommandBuilder.h
--- a/Engine/Plugins/Tests/CQTest/Source/CQTest/Public/Commands/TestCommandBuilder.h
+++ b/Engine/Plugins/Tests/CQTest/Source/CQTest/Public/Commands/TestCommandBuilder.h
@@ -2,6 +2,7 @@
 
 #pragma once
 
+#include "Algo/Reverse.h"
 #include "Misc/AutomationTest.h"
 #include "Commands/TestCommands.h"
 
@@ -64,28 +65,66 @@ public:
 		return Until(Description, Query, Timeout);
 	}
 
+	FTestCommandBuilder& OnTearDown(const TCHAR* Description, TFunction<void()> Action)
+	{
+		if (!TestRunner.HasAnyErrors())
+		{
+			TearDownQueue.Add(MakeShared<FExecute>(TestRunner, Action, Description, ECQTestFailureBehavior::Run));
+		}
+		return *this;
+	}
+
+	FTestCommandBuilder& OnTearDown(TFunction<void()> Action)
+	{
+		return OnTearDown(nullptr, Action);
+	}
+
+	FTestCommandBuilder& CleanUpWith(const TCHAR* Description, TFunction<void()> Action)
+	{
+		return OnTearDown(Description, Action);
+	}
+
+	FTestCommandBuilder& CleanUpWith(TFunction<void()> Action)
+	{
+		return OnTearDown(nullptr, Action);
+	}
+
 	TSharedPtr<IAutomationLatentCommand> Build()
+	{
+		return BuildQueue(CommandQueue);
+	}
+
+	TSharedPtr<IAutomationLatentCommand> BuildTearDown()
+	{
+		// Last in, first out
+		Algo::Reverse(TearDownQueue);
+		return BuildQueue(TearDownQueue);
+	}
+
+private:
+	TSharedPtr<IAutomationLatentCommand> BuildQueue(TArray<TSharedPtr<IAutomationLatentCommand>>& Queue)
 	{
 		TSharedPtr<IAutomationLatentCommand> Result = nullptr;
-		if (CommandQueue.Num() == 0)
+		if (Queue.Num() == 0)
 		{
 			return Result;
 		}
-		else if (CommandQueue.Num() == 1)
+		else if (Queue.Num() == 1)
 		{
-			Result = CommandQueue[0];
+			Result = Queue[0];
 		}
 		else
 		{
-			Result = MakeShared<FRunSequence>(CommandQueue);
+			Result = MakeShared<FRunSequence>(Queue);
 		}
 
-		CommandQueue.Empty();
+		Queue.Empty();
 		return Result;
 	}
 
 protected:
 	TArray<TSharedPtr<IAutomationLatentCommand>> CommandQueue{};
+	TArray<TSharedPtr<IAutomationLatentCommand>> TearDownQueue{};
 	FAutomationTestBase& TestRunner;
 
 	template<typename Asserter>
diff --git a/Engine/Plugins/Tests/CQTest/Source/CQTest/Public/Commands/TestCommands.h b/Engine/Plugins/Tests/CQTest/Source/CQTest/Public/Commands/TestCommands.h
--- a/Engine/Plugins/Tests/CQTest/Source/CQTest/Public/Commands/TestCommands.h
+++ b/Engine/Plugins/Tests/CQTest/Source/CQTest/Public/Commands/TestCommands.h
@@ -25,13 +25,20 @@ public:
 	bool bHasLoggedStart = false;
 };
 
+enum class ECQTestFailureBehavior
+{
+	Skip,
+	Run
+};
+
 class CQTEST_API FExecute : public IAutomationLatentCommand
 {
 public:
-	FExecute(FAutomationTestBase& InTestRunner, TFunction<void()> Func, const TCHAR* InDescription = nullptr)
+	FExecute(FAutomationTestBase& InTestRunner, TFunction<void()> Func, const TCHAR* InDescription = nullptr, ECQTestFailureBehavior InFailureBehavior = ECQTestFailureBehavior::Skip)
 		: TestRunner(InTestRunner)
 		, Func(MoveTemp(Func)) 
 		, Description(InDescription)
+		, FailureBehavior(InFailureBehavior)
 	{}
 
 	bool Update() override;
@@ -39,6 +46,7 @@ public:
 	FAutomationTestBase& TestRunner;
 	TFunction<void()> Func;
 	const TCHAR* Description = nullptr;
+	ECQTestFailureBehavior FailureBehavior;
 };
 
 class CQTEST_API FRunSequence : public IAutomationLatentCommand
diff --git a/Engine/Plugins/Tests/CQTest/Source/CQTest/Public/Components/PIENetworkComponent.h b/Engine/Plugins/Tests/CQTest/Source/CQTest/Public/Components/PIENetworkComponent.h
--- a/Engine/Plugins/Tests/CQTest/Source/CQTest/Public/Components/PIENetworkComponent.h
+++ b/Engine/Plugins/Tests/CQTest/Source/CQTest/Public/Components/PIENetworkComponent.h
@@ -26,7 +26,6 @@ class CQTEST_API FBasePIENetworkComponent
 {
 public:
 	FBasePIENetworkComponent(FAutomationTestBase* InTestRunner, FTestCommandBuilder& InCommandBuilder, bool IsInitializing);
-	~FBasePIENetworkComponent();
 
 	FBasePIENetworkComponent& Then(TFunction<void()> Action);
 	FBasePIENetworkComponent& Then(const TCHAR* Description, TFunction<void()> Action);
@@ -44,6 +43,7 @@ protected:
 	void StartPie();
 	bool CollectPieWorlds();
 	bool AwaitConnections();
+	void RestoreState();
 
 	TUniquePtr<FBasePIENetworkComponentState> ServerState = nullptr;
 	TArray<TUniquePtr<FBasePIENetworkComponentState>> ClientStates;
diff --git a/Engine/Plugins/Tests/CQTest/Source/CQTest/Public/Components/PIENetworkTestStateRestorer.h b/Engine/Plugins/Tests/CQTest/Source/CQTest/Public/Components/PIENetworkTestStateRestorer.h
--- a/Engine/Plugins/Tests/CQTest/Source/CQTest/Public/Components/PIENetworkTestStateRestorer.h
+++ b/Engine/Plugins/Tests/CQTest/Source/CQTest/Public/Components/PIENetworkTestStateRestorer.h
@@ -18,7 +18,6 @@ public:
 	void Restore();
 
 private:
-	UWorld* EditorWorld = nullptr;
 	bool SetWasLoadedFlag = false;
 	FSoftClassPath OriginalGameInstance = FSoftClassPath();
 	TSubclassOf<AGameModeBase> OriginalGameMode = nullptr;
diff --git a/Engine/Plugins/Tests/CQTest/Source/CQTest/Public/Impl/CQTest.inl b/Engine/Plugins/Tests/CQTest/Source/CQTest/Public/Impl/CQTest.inl
--- a/Engine/Plugins/Tests/CQTest/Source/CQTest/Public/Impl/CQTest.inl
+++ b/Engine/Plugins/Tests/CQTest/Source/CQTest/Public/Impl/CQTest.inl
@@ -16,6 +16,12 @@ namespace
 		{
 			UE_LOG(LogCqTest, Log, TEXT("Before Test"));
 
+			if (GEngine != nullptr)
+			{
+				//Do not collect garbage during the test. We force GC at the end.
+				GEngine->DelayGarbageCollection();
+			}
+
 			CurrentTest.Setup();
 			if (auto LatentActions = CurrentTest.TestCommandBuilder.Build())
 			{
@@ -78,6 +84,10 @@ namespace
 			UE_LOG(LogCqTest, Log, TEXT("Running After Test"));
 
 			CurrentTest.TearDown();
+			if (auto LatentActions = CurrentTest.TestCommandBuilder.BuildTearDown())
+			{
+				Sequence->Prepend(LatentActions);
+			}
 			if (auto LatentActions = CurrentTest.TestCommandBuilder.Build())
 			{
 				Sequence->Prepend(LatentActions);
@@ -101,6 +111,13 @@ namespace
 			UE_LOG(LogCqTest, Log, TEXT("Tearing Down Test"));
 
 			TestRunner.CurrentTestPtr = nullptr;
+
+			if (GEngine != nullptr)
+			{
+				//Force GC at the end of every test.
+				GEngine->ForceGarbageCollection();
+			}
+
 			return true;
 		}
 
diff --git a/Engine/Plugins/Tests/CQTest/Source/CQTestTests/Private/CQTestTests.cpp b/Engine/Plugins/Tests/CQTest/Source/CQTestTests/Private/CQTestTests.cpp
--- a/Engine/Plugins/Tests/CQTest/Source/CQTestTests/Private/CQTestTests.cpp
+++ b/Engine/Plugins/Tests/CQTest/Source/CQTestTests/Private/CQTestTests.cpp
@@ -329,6 +329,32 @@ namespace CQTestTests
 		}
 	};
 
+
+	TEST_CLASS(LatentActionOnTearDown, "TestFramework.CQTest.Core")
+	{
+		FString ExpectedError = TEXT("ExpectedErrorShouldBeCleared");
+
+		TEST_METHOD(Teardown_Fires_AfterAll)
+		{
+			TestCommandBuilder.Do([&]() {
+				Assert.Fail(ExpectedError);
+			})
+			.OnTearDown([&]() {
+				ClearExpectedError(*TestRunner, ExpectedError);
+			});
+		}
+
+		TEST_METHOD(Teardown_Fires_InReverseOrder)
+		{
+			TestCommandBuilder.OnTearDown([&]() {
+				ClearExpectedError(*TestRunner, ExpectedError);
+			}).OnTearDown([&]() {
+				Assert.Fail(ExpectedError);
+			});
+		}
+	};
+
+
 	TEST_CLASS(LatentActionErrors, "TestFramework.CQTest.Core")
 	{
 		FString ExpectedError = TEXT("ExpectedError");
