Fix PIENetworkComponent.cpp(89): error C2653: 'FAutomationEditorCommonUtils': is not a class or namespace name on CookedEditor
ENABLE_PIE_NETWORK_TEST needs WITH_AUTOMATION_TESTS

#rnx

diff --git a/Engine/Plugins/Tests/CQTest/Source/CQTest/Public/Components/PIENetworkComponent.h b/Engine/Plugins/Tests/CQTest/Source/CQTest/Public/Components/PIENetworkComponent.h
--- a/Engine/Plugins/Tests/CQTest/Source/CQTest/Public/Components/PIENetworkComponent.h
+++ b/Engine/Plugins/Tests/CQTest/Source/CQTest/Public/Components/PIENetworkComponent.h
@@ -3,7 +3,7 @@
 #pragma once
 #include "HAL/Platform.h"
 
-#define ENABLE_PIE_NETWORK_TEST WITH_EDITOR
+#define ENABLE_PIE_NETWORK_TEST WITH_EDITOR && WITH_AUTOMATION_TESTS
 
 #if ENABLE_PIE_NETWORK_TEST
 
