fix static analysis warning in PIENetworkComponent

#rnx

diff --git a/Engine/Plugins/Tests/CQTest/Source/CQTest/Public/Components/PIENetworkComponent.h b/Engine/Plugins/Tests/CQTest/Source/CQTest/Public/Components/PIENetworkComponent.h
--- a/Engine/Plugins/Tests/CQTest/Source/CQTest/Public/Components/PIENetworkComponent.h
+++ b/Engine/Plugins/Tests/CQTest/Source/CQTest/Public/Components/PIENetworkComponent.h
@@ -3,7 +3,7 @@
 #pragma once
 #include "HAL/Platform.h"
 
-#define ENABLE_PIE_NETWORK_TEST WITH_EDITOR && WITH_AUTOMATION_TESTS
+#define ENABLE_PIE_NETWORK_TEST WITH_EDITOR && (WITH_DEV_AUTOMATION_TESTS || WITH_PERF_AUTOMATION_TESTS)
 
 #if ENABLE_PIE_NETWORK_TEST
 
