CQTest - Implement FActorTestSpawner::GetGameIntance()

#rnx

diff --git a/Engine/Plugins/Tests/CQTest/Source/CQTest/Private/Components/ActorTestSpawner.cpp b/Engine/Plugins/Tests/CQTest/Source/CQTest/Private/Components/ActorTestSpawner.cpp
--- a/Engine/Plugins/Tests/CQTest/Source/CQTest/Private/Components/ActorTestSpawner.cpp
+++ b/Engine/Plugins/Tests/CQTest/Source/CQTest/Private/Components/ActorTestSpawner.cpp
@@ -96,3 +96,8 @@ void FActorTestSpawner::InitializeGameSubsystems()
 	GameInstance->AddToRoot();
 	GameInstance->InitializeForTest(GetWorld());
 }
+
+UTestGameInstance* FActorTestSpawner::GetGameInstance()
+{
+	return GameInstance;
+}
diff --git a/Engine/Plugins/Tests/CQTest/Source/CQTest/Public/Components/ActorTestSpawner.h b/Engine/Plugins/Tests/CQTest/Source/CQTest/Public/Components/ActorTestSpawner.h
--- a/Engine/Plugins/Tests/CQTest/Source/CQTest/Public/Components/ActorTestSpawner.h
+++ b/Engine/Plugins/Tests/CQTest/Source/CQTest/Public/Components/ActorTestSpawner.h
@@ -14,6 +14,8 @@ struct CQTEST_API FActorTestSpawner : public FSpawnHelper
 
 	void InitializeGameSubsystems();
 
+	UTestGameInstance* GetGameInstance();
+
 protected:
 	virtual UWorld* CreateWorld() override;
 
