Introduce ASSERT_FAIL macro alias to Assert.Fail(Msg); return;

#rnx

diff --git a/Engine/Plugins/Tests/CQTest/Source/CQTest/Public/CQTest.h b/Engine/Plugins/Tests/CQTest/Source/CQTest/Public/CQTest.h
--- a/Engine/Plugins/Tests/CQTest/Source/CQTest/Public/CQTest.h
+++ b/Engine/Plugins/Tests/CQTest/Source/CQTest/Public/CQTest.h
@@ -167,6 +167,7 @@ struct TTest : TBaseTest<AsserterType>
 	void _TestName::_TestName##_Method()
 
 #define ASSERT_THAT(_assertion) if (!Assert._assertion) {return;}
+#define ASSERT_FAIL(_Msg) Assert.Fail(_Msg); return;
 
 #define BEFORE_EACH() virtual void Setup() override
 #define AFTER_EACH() virtual void TearDown() override
diff --git a/Engine/Plugins/Tests/CQTest/Source/CQTestTests/Private/Assert/AssertionTests.cpp b/Engine/Plugins/Tests/CQTest/Source/CQTestTests/Private/Assert/AssertionTests.cpp
--- a/Engine/Plugins/Tests/CQTest/Source/CQTestTests/Private/Assert/AssertionTests.cpp
+++ b/Engine/Plugins/Tests/CQTest/Source/CQTestTests/Private/Assert/AssertionTests.cpp
@@ -17,6 +17,13 @@ namespace CQTests
 			Assert.Fail(TEXT(""));
 		}
 
+		TEST_METHOD(AssertFail_Macro_AddsErrorAndExits)
+		{
+			Assert.ExpectError(ExpectedError);
+			ASSERT_FAIL(ExpectedError);
+			ASSERT_FAIL(TEXT("It should not assert after the first expected assert since it exits"));
+		}
+
 		TEST_METHOD(Assertions_Accept_RawStrings)
 		{
 			Assert.ExpectError("Hello World");
