CQTest - better support enum uint8/int8 when converting to string

#rnx

diff --git a/Engine/Plugins/Tests/CQTest/Source/CQTest/Public/Assert/CQTestConvert.h b/Engine/Plugins/Tests/CQTest/Source/CQTest/Public/Assert/CQTestConvert.h
--- a/Engine/Plugins/Tests/CQTest/Source/CQTest/Public/Assert/CQTestConvert.h
+++ b/Engine/Plugins/Tests/CQTest/Source/CQTest/Public/Assert/CQTestConvert.h
@@ -92,7 +92,15 @@ namespace CQTestConvert
 		}
 		else if constexpr (std::is_enum_v<T>)
 		{
-			return ToString(static_cast<std::underlying_type_t<T>>(Input));
+			// Check special case of uint8/int8 to avoid interpretation as a char
+			if constexpr (std::is_same_v<uint8, std::underlying_type_t<T>> || std::is_same_v<int8, std::underlying_type_t<T>>)
+			{
+				return ToString(static_cast<int>(Input));
+			}
+			else
+			{
+				return ToString(static_cast<std::underlying_type_t<T>>(Input));
+			}
 		}
 		else
 		{
