[CQTest] Add interface to set SuppressWarnings and SuppressErrors from TestRunner

#rnx

diff --git a/Engine/Plugins/Tests/CQTest/Source/CQTest/Public/CQTest.h b/Engine/Plugins/Tests/CQTest/Source/CQTest/Public/CQTest.h
--- a/Engine/Plugins/Tests/CQTest/Source/CQTest/Public/CQTest.h
+++ b/Engine/Plugins/Tests/CQTest/Source/CQTest/Public/CQTest.h
@@ -43,6 +43,13 @@ namespace TestDirectoryGenerator
 	CQTEST_API FString Generate(const FString& Filename);
 }
 
+enum class ECQTestSuppressLogBehavior
+{
+	Default,
+	True,
+	False
+};
+
 static const FString GenerateTestDirectory = TEXT("");
 static constexpr uint32 DefaultFlags = EAutomationTestFlags::ApplicationContextMask | EAutomationTestFlags::ProductFilter;
 
@@ -87,10 +94,15 @@ struct TTestRunner : public FAutomationTestBase
 	FString GetTestSourceFileName() const override;
 	int32 GetTestSourceFileLine() const override;
 	int32 GetTestSourceFileLine(const FString& Name) const override;
+	bool SuppressLogWarnings() override;
+	bool SuppressLogErrors() override;
 	void GetTests(TArray<FString>& OutBeautifiedNames, TArray<FString>& OutTestCommands) const override;
 
 	bool RunTest(const FString& RequestedTest) override;
 
+	virtual void SetSuppressLogWarnings(ECQTestSuppressLogBehavior Behavior = ECQTestSuppressLogBehavior::True);
+	virtual void SetSuppressLogErrors(ECQTestSuppressLogBehavior Behavior = ECQTestSuppressLogBehavior::True);
+
 	int32 LineNumber = 0;
 	FString FileName;
 	FString TestDir;
@@ -105,6 +117,10 @@ struct TTestRunner : public FAutomationTestBase
 protected:
 	FString GetBeautifiedTestName() const override;
 	uint32 GetRequiredDeviceNum() const override;
+
+	ECQTestSuppressLogBehavior SuppressLogWarningsBehavior{ ECQTestSuppressLogBehavior::Default };
+	ECQTestSuppressLogBehavior SuppressLogErrorsBehavior{ ECQTestSuppressLogBehavior::Default };
+
 };
 
 template <typename Derived, typename AsserterType>
diff --git a/Engine/Plugins/Tests/CQTest/Source/CQTest/Public/Impl/CQTest.inl b/Engine/Plugins/Tests/CQTest/Source/CQTest/Public/Impl/CQTest.inl
--- a/Engine/Plugins/Tests/CQTest/Source/CQTest/Public/Impl/CQTest.inl
+++ b/Engine/Plugins/Tests/CQTest/Source/CQTest/Public/Impl/CQTest.inl
@@ -110,8 +110,10 @@ namespace
 		{
 			UE_LOG(LogCqTest, Log, TEXT("Tearing Down Test"));
 
-			TestRunner.CurrentTestPtr = nullptr;
+			TestRunner.SetSuppressLogWarnings(ECQTestSuppressLogBehavior::Default);
+			TestRunner.SetSuppressLogErrors(ECQTestSuppressLogBehavior::Default);
 
+			TestRunner.CurrentTestPtr = nullptr;
 			if (GEngine != nullptr)
 			{
 				//Force GC at the end of every test.
@@ -231,6 +233,38 @@ inline int32 TTestRunner<AsserterType>::GetTestSourceFileLine(const FString& Nam
 	return GetTestSourceFileLine();
 }
 
+template <typename AsserterType>
+inline bool TTestRunner<AsserterType>::SuppressLogWarnings()
+{
+	if (SuppressLogWarningsBehavior == ECQTestSuppressLogBehavior::Default)
+	{
+		return bSuppressLogWarnings;
+	}
+	return SuppressLogWarningsBehavior == ECQTestSuppressLogBehavior::True;
+}
+
+template <typename AsserterType>
+inline bool TTestRunner<AsserterType>::SuppressLogErrors()
+{
+	if (SuppressLogErrorsBehavior == ECQTestSuppressLogBehavior::Default)
+	{
+		return bSuppressLogErrors;
+	}
+	return SuppressLogErrorsBehavior == ECQTestSuppressLogBehavior::True;
+}
+
+template <typename AsserterType>
+inline void TTestRunner<AsserterType>::SetSuppressLogWarnings(ECQTestSuppressLogBehavior Behavior)
+{
+	SuppressLogWarningsBehavior = Behavior;
+}
+
+template <typename AsserterType>
+inline void TTestRunner<AsserterType>::SetSuppressLogErrors(ECQTestSuppressLogBehavior Behavior)
+{
+	SuppressLogErrorsBehavior = Behavior;
+}
+
 template <typename AsserterType>
 inline void TTestRunner<AsserterType>::GetTests(TArray<FString>& OutBeautifiedNames, TArray<FString>& OutTestCommands) const
 {
diff --git a/Engine/Plugins/Tests/CQTest/Source/CQTestTests/Private/CQTestTests.cpp b/Engine/Plugins/Tests/CQTest/Source/CQTestTests/Private/CQTestTests.cpp
--- a/Engine/Plugins/Tests/CQTest/Source/CQTestTests/Private/CQTestTests.cpp
+++ b/Engine/Plugins/Tests/CQTest/Source/CQTestTests/Private/CQTestTests.cpp
@@ -427,4 +427,61 @@ namespace CQTestTests
 			ASSERT_THAT(IsTrue(Tickable.TickCount > 2));
 		}
 	};
+
+	DEFINE_LOG_CATEGORY_STATIC(TestSuppressLog, Log, All);
+
+	TEST_CLASS(SuppressWarningsAndErrors, "TestFramework.CQTest.Core")
+	{
+		bool bExpectError{ false };
+		bool bExpectWarning{ false };
+		FString ExpectedError = TEXT("ExpectedError");
+
+		AFTER_EACH()
+		{
+			static FString FormattedMessage = FString::Printf(TEXT("TestSuppressLog: %s"), *ExpectedError);
+			if (bExpectError)
+			{
+				ClearExpectedErrors(*this->TestRunner, { TEXT("will be marked as failing due to errors"), FormattedMessage });
+			}
+			if (bExpectWarning)
+			{
+				ClearExpectedWarning(*this->TestRunner, FormattedMessage);
+			}
+		}
+
+		TEST_METHOD(SuppressWarnings_WithWarnings_SuppressesWarning)
+		{
+			TestRunner->SetSuppressLogWarnings();
+			UE_LOG(TestSuppressLog, Warning, TEXT("This log should be suppressed"));
+		}
+
+		TEST_METHOD(SuppressWarnings_WithError_DoesNotSuppressError)
+		{
+			TestRunner->SetSuppressLogWarnings();
+			bExpectError = true;
+			UE_LOG(TestSuppressLog, Error, TEXT("%s"), *ExpectedError);
+		}
+
+		TEST_METHOD(SuppressErrors_WithWarnings_DoesNotSuppressWarning)
+		{
+			TestRunner->SetSuppressLogErrors();
+			bExpectWarning = true;
+			UE_LOG(TestSuppressLog, Warning, TEXT("%s"), *ExpectedError);
+		}
+
+		TEST_METHOD(SuppressErrors_WithError_SuppressesError)
+		{
+			TestRunner->SetSuppressLogErrors();
+			UE_LOG(TestSuppressLog, Error, TEXT("This log should be suppressed"));
+		}
+
+		TEST_METHOD(SuppressErrors_ChangedMidTest_RespectsChanges)
+		{
+			TestRunner->SetSuppressLogErrors();
+			UE_LOG(TestSuppressLog, Error, TEXT("This log should be suppressed"));
+			TestRunner->SetSuppressLogErrors(ECQTestSuppressLogBehavior::False);
+			bExpectError = true;
+			UE_LOG(TestSuppressLog, Error, TEXT("%s"), *ExpectedError);
+		}
+	};
 } // namespace CQTestTests
diff --git a/Engine/Plugins/Tests/CQTest/Source/CQTestTests/Private/CQTestUnitTestHelper.cpp b/Engine/Plugins/Tests/CQTest/Source/CQTestTests/Private/CQTestUnitTestHelper.cpp
--- a/Engine/Plugins/Tests/CQTest/Source/CQTestTests/Private/CQTestUnitTestHelper.cpp
+++ b/Engine/Plugins/Tests/CQTest/Source/CQTestTests/Private/CQTestUnitTestHelper.cpp
@@ -2,19 +2,52 @@
 
 #include "CQTestUnitTestHelper.h"
 
-void ClearExpectedError(FAutomationTestBase& TestRunner, const FString& ExpectedError)
+void ClearExpectedErrors(FAutomationTestBase& TestRunner, const TArray<FString>& ExpectedErrors)
 {
 	FAutomationTestExecutionInfo testInfo;
 	TestRunner.GetExecutionInfo(testInfo);
-	if (testInfo.GetErrorTotal() != 1)
+	if (testInfo.GetErrorTotal() == 0)
 	{
 		return;
 	}
-	testInfo.RemoveAllEvents([&ExpectedError](FAutomationEvent& event) {
-		return event.Message.Equals(ExpectedError);
-		});
+	for(auto & ExpectedError : ExpectedErrors)
+	{
+		testInfo.RemoveAllEvents([&ExpectedError](FAutomationEvent& event) {
+			return event.Message.Contains(ExpectedError);
+			});
+	}
 	if (testInfo.GetErrorTotal() == 0)
 	{
 		TestRunner.ClearExecutionInfo();
 	}
 }
+
+void ClearExpectedError(FAutomationTestBase& TestRunner, const FString& ExpectedError)
+{
+	ClearExpectedErrors(TestRunner, { ExpectedError });
+}
+
+void ClearExpectedWarnings(FAutomationTestBase& TestRunner, const TArray<FString>& ExpectedWarnings)
+{
+	FAutomationTestExecutionInfo testInfo;
+	TestRunner.GetExecutionInfo(testInfo);
+	if (testInfo.GetWarningTotal() == 0)
+	{
+		return;
+	}
+	for (auto& ExpectedWarning : ExpectedWarnings)
+	{
+		testInfo.RemoveAllEvents([&ExpectedWarning](FAutomationEvent& event) {
+			return event.Message.Contains(ExpectedWarning);
+			});
+	}
+	if (testInfo.GetWarningTotal() == 0)
+	{
+		TestRunner.ClearExecutionInfo();
+	}
+}
+
+void ClearExpectedWarning(FAutomationTestBase& TestRunner, const FString& ExpectedWarning)
+{
+	ClearExpectedWarnings(TestRunner, { ExpectedWarning });
+}
\ No newline at end of file
diff --git a/Engine/Plugins/Tests/CQTest/Source/CQTestTests/Private/CQTestUnitTestHelper.h b/Engine/Plugins/Tests/CQTest/Source/CQTestTests/Private/CQTestUnitTestHelper.h
--- a/Engine/Plugins/Tests/CQTest/Source/CQTestTests/Private/CQTestUnitTestHelper.h
+++ b/Engine/Plugins/Tests/CQTest/Source/CQTestTests/Private/CQTestUnitTestHelper.h
@@ -6,4 +6,7 @@
 // When testing the framework, occasionally we need the error to live on the test framework
 // for short-circuiting logic.  We use this method to manually clear the error
 // (Normally from AFTER_EACH) to cause the test to pass
+void ClearExpectedErrors(FAutomationTestBase& TestRunner, const TArray<FString>& ExpectedErrors);
 void ClearExpectedError(FAutomationTestBase& TestRunner, const FString& ExpectedError);
+void ClearExpectedWarnings(FAutomationTestBase& TestRunner, const TArray<FString>& ExpectedWarnings);
+void ClearExpectedWarning(FAutomationTestBase& TestRunner, const FString& ExpectedWarning);
