[CQTest] Use UE_NODISCARD macro instead of [[nodiscard]]

#rnx

diff --git a/Engine/Plugins/Tests/CQTest/Source/CQTest/Public/Assert/NoDiscardAsserter.h b/Engine/Plugins/Tests/CQTest/Source/CQTest/Public/Assert/NoDiscardAsserter.h
--- a/Engine/Plugins/Tests/CQTest/Source/CQTest/Public/Assert/NoDiscardAsserter.h
+++ b/Engine/Plugins/Tests/CQTest/Source/CQTest/Public/Assert/NoDiscardAsserter.h
@@ -18,75 +18,75 @@ public:
 
 	void Fail(FString error);
 
-	[[nodiscard]] bool IsTrue(bool Condition);
-	[[nodiscard]] bool IsTrue(bool Condition, const char* FailureMessage);
-	[[nodiscard]] bool IsTrue(bool Condition, const TCHAR* FailureMessage);
-	[[nodiscard]] bool IsTrue(bool Condition, const FString& FailureMessage);
+	UE_NODISCARD bool IsTrue(bool Condition);
+	UE_NODISCARD bool IsTrue(bool Condition, const char* FailureMessage);
+	UE_NODISCARD bool IsTrue(bool Condition, const TCHAR* FailureMessage);
+	UE_NODISCARD bool IsTrue(bool Condition, const FString& FailureMessage);
 
-	[[nodiscard]] bool IsFalse(bool Condition);
-	[[nodiscard]] bool IsFalse(bool Condition, const char* FailureMessage);
-	[[nodiscard]] bool IsFalse(bool Condition, const TCHAR* FailureMessage);
-	[[nodiscard]] bool IsFalse(bool Condition, const FString& FailureMessage);
+	UE_NODISCARD bool IsFalse(bool Condition);
+	UE_NODISCARD bool IsFalse(bool Condition, const char* FailureMessage);
+	UE_NODISCARD bool IsFalse(bool Condition, const TCHAR* FailureMessage);
+	UE_NODISCARD bool IsFalse(bool Condition, const FString& FailureMessage);
 
 	template <typename T>
-	[[nodiscard]] bool IsNull(const T& Ptr);
+	UE_NODISCARD bool IsNull(const T& Ptr);
 
 	template <typename T>
-	[[nodiscard]] bool IsNull(const T& Ptr, const char* FailureMessage);
+	UE_NODISCARD bool IsNull(const T& Ptr, const char* FailureMessage);
 
 	template <typename T>
-	[[nodiscard]] bool IsNull(const T& Ptr, const TCHAR* FailureMessage);
+	UE_NODISCARD bool IsNull(const T& Ptr, const TCHAR* FailureMessage);
 
 	template <typename T>
-	[[nodiscard]] bool IsNull(const T& Ptr, const FString& FailureMessage);
+	UE_NODISCARD bool IsNull(const T& Ptr, const FString& FailureMessage);
 
 	template <typename T>
-	[[nodiscard]] bool IsNotNull(const T& Ptr);
+	UE_NODISCARD bool IsNotNull(const T& Ptr);
 
 	template <typename T>
-	[[nodiscard]] bool IsNotNull(const T& Ptr, const char* FailureMessage);
+	UE_NODISCARD bool IsNotNull(const T& Ptr, const char* FailureMessage);
 
 	template <typename T>
-	[[nodiscard]] bool IsNotNull(const T& Ptr, const TCHAR* FailureMessage);
+	UE_NODISCARD bool IsNotNull(const T& Ptr, const TCHAR* FailureMessage);
 
 	template <typename T>
-	[[nodiscard]] bool IsNotNull(const T& Ptr, const FString& FailureMessage);
+	UE_NODISCARD bool IsNotNull(const T& Ptr, const FString& FailureMessage);
 
 	template <typename TExpected, typename TActual>
-	[[nodiscard]] bool AreEqual(const TExpected& Expected, const TActual& Actual);
+	UE_NODISCARD bool AreEqual(const TExpected& Expected, const TActual& Actual);
 
 	template <typename TExpected, typename TActual>
-	[[nodiscard]] bool AreEqual(const TExpected& Expected, const TActual& Actual, const char* FailureMessage);
+	UE_NODISCARD bool AreEqual(const TExpected& Expected, const TActual& Actual, const char* FailureMessage);
 
 	template <typename TExpected, typename TActual>
-	[[nodiscard]] bool AreEqual(const TExpected& Expected, const TActual& Actual, const TCHAR* FailureMessage);
+	UE_NODISCARD bool AreEqual(const TExpected& Expected, const TActual& Actual, const TCHAR* FailureMessage);
 
 	template <typename TExpected, typename TActual>
-	[[nodiscard]] bool AreEqual(const TExpected& Expected, const TActual& Actual, const FString& FailureMessage);
+	UE_NODISCARD bool AreEqual(const TExpected& Expected, const TActual& Actual, const FString& FailureMessage);
 
 	template <typename TExpected, typename TActual>
-	[[nodiscard]] bool AreNotEqual(const TExpected& Expected, const TActual& Actual);
+	UE_NODISCARD bool AreNotEqual(const TExpected& Expected, const TActual& Actual);
 
 	template <typename TExpected, typename TActual>
-	[[nodiscard]] bool AreNotEqual(const TExpected& Expected, const TActual& Actual, const char* FailureMessage);
+	UE_NODISCARD bool AreNotEqual(const TExpected& Expected, const TActual& Actual, const char* FailureMessage);
 
 	template <typename TExpected, typename TActual>
-	[[nodiscard]] bool AreNotEqual(const TExpected& Expected, const TActual& Actual, const TCHAR* FailureMessage);
+	UE_NODISCARD bool AreNotEqual(const TExpected& Expected, const TActual& Actual, const TCHAR* FailureMessage);
 
 	template <typename TExpected, typename TActual>
-	[[nodiscard]] bool AreNotEqual(const TExpected& Expected, const TActual& Actual, const FString& FailureMessage);
+	UE_NODISCARD bool AreNotEqual(const TExpected& Expected, const TActual& Actual, const FString& FailureMessage);
 
 	template <typename TExpected, typename TActual, typename TEpsilon>
-	[[nodiscard]] bool IsNear(TExpected Expected, TActual Actual, TEpsilon Epsilon);
+	UE_NODISCARD bool IsNear(TExpected Expected, TActual Actual, TEpsilon Epsilon);
 
 	template <typename TExpected, typename TActual, typename TEpsilon>
-	[[nodiscard]] bool IsNear(TExpected Expected, TActual Actual, TEpsilon Epsilon, const char* FailureMessage);
+	UE_NODISCARD bool IsNear(TExpected Expected, TActual Actual, TEpsilon Epsilon, const char* FailureMessage);
 
 	template <typename TExpected, typename TActual, typename TEpsilon>
-	[[nodiscard]] bool IsNear(TExpected Expected, TActual Actual, TEpsilon Epsilon, const TCHAR* FailureMessage);
+	UE_NODISCARD bool IsNear(TExpected Expected, TActual Actual, TEpsilon Epsilon, const TCHAR* FailureMessage);
 
 	template <typename TExpected, typename TActual, typename TEpsilon>
-	[[nodiscard]] bool IsNear(TExpected Expected, TActual Actual, TEpsilon Epsilon, const FString& FailureMessage);
+	UE_NODISCARD bool IsNear(TExpected Expected, TActual Actual, TEpsilon Epsilon, const FString& FailureMessage);
 };
 
 #include "Assert/NoDiscardAsserter.inl"
\ No newline at end of file
diff --git a/Engine/Plugins/Tests/CQTest/Source/CQTest/Public/Assert/NoDiscardAsserter.inl b/Engine/Plugins/Tests/CQTest/Source/CQTest/Public/Assert/NoDiscardAsserter.inl
--- a/Engine/Plugins/Tests/CQTest/Source/CQTest/Public/Assert/NoDiscardAsserter.inl
+++ b/Engine/Plugins/Tests/CQTest/Source/CQTest/Public/Assert/NoDiscardAsserter.inl
@@ -29,66 +29,66 @@ inline void FNoDiscardAsserter::Fail(FString Error)
 	TestRunner.AddError(Error);
 }
 
-[[nodiscard]] inline bool FNoDiscardAsserter::IsTrue(bool Condition)
+UE_NODISCARD inline bool FNoDiscardAsserter::IsTrue(bool Condition)
 {
 	return IsTrue(Condition, "condition");
 }
 
-[[nodiscard]] inline bool FNoDiscardAsserter::IsTrue(bool Condition, const char* FailureMessage)
+UE_NODISCARD inline bool FNoDiscardAsserter::IsTrue(bool Condition, const char* FailureMessage)
 {
 	return IsTrue(Condition, FString(FailureMessage));
 }
 
-[[nodiscard]] inline bool FNoDiscardAsserter::IsTrue(bool Condition, const TCHAR* FailureMessage)
+UE_NODISCARD inline bool FNoDiscardAsserter::IsTrue(bool Condition, const TCHAR* FailureMessage)
 {
 	return IsTrue(Condition, FString(FailureMessage));
 }
 
-[[nodiscard]] inline bool FNoDiscardAsserter::IsTrue(bool Condition, const FString& FailureMessage)
+UE_NODISCARD inline bool FNoDiscardAsserter::IsTrue(bool Condition, const FString& FailureMessage)
 {
 	return TestRunner.TestTrue(FailureMessage, Condition);
 }
 
-[[nodiscard]] inline bool FNoDiscardAsserter::IsFalse(bool Condition)
+UE_NODISCARD inline bool FNoDiscardAsserter::IsFalse(bool Condition)
 {
 	return IsFalse(Condition, "condition");
 }
 
-[[nodiscard]] inline bool FNoDiscardAsserter::IsFalse(bool Condition, const char* FailureMessage)
+UE_NODISCARD inline bool FNoDiscardAsserter::IsFalse(bool Condition, const char* FailureMessage)
 {
 	return IsFalse(Condition, FString(FailureMessage));
 }
 
-[[nodiscard]] inline bool FNoDiscardAsserter::IsFalse(bool Condition, const TCHAR* FailureMessage)
+UE_NODISCARD inline bool FNoDiscardAsserter::IsFalse(bool Condition, const TCHAR* FailureMessage)
 {
 	return IsFalse(Condition, FString(FailureMessage));
 }
 
-[[nodiscard]] inline bool FNoDiscardAsserter::IsFalse(bool Condition, const FString& FailureMessage)
+UE_NODISCARD inline bool FNoDiscardAsserter::IsFalse(bool Condition, const FString& FailureMessage)
 {
 	return TestRunner.TestFalse(FString(FailureMessage), Condition);
 }
 
 template <typename T>
-[[nodiscard]] bool FNoDiscardAsserter::IsNull(const T& Ptr)
+UE_NODISCARD bool FNoDiscardAsserter::IsNull(const T& Ptr)
 {
 	return IsNull(Ptr, "TestNull failed");
 }
 
 template <typename T>
-[[nodiscard]] bool FNoDiscardAsserter::IsNull(const T& Ptr, const char* FailureMessage)
+UE_NODISCARD bool FNoDiscardAsserter::IsNull(const T& Ptr, const char* FailureMessage)
 {
 	return IsNull(Ptr, FString(FailureMessage));
 }
 
 template <typename T>
-[[nodiscard]] bool FNoDiscardAsserter::IsNull(const T& Ptr, const TCHAR* FailureMessage)
+UE_NODISCARD bool FNoDiscardAsserter::IsNull(const T& Ptr, const TCHAR* FailureMessage)
 {
 	return IsNull(Ptr, FString(FailureMessage));
 }
 
 template <typename T>
-[[nodiscard]] bool FNoDiscardAsserter::IsNull(const T& Ptr, const FString& FailureMessage)
+UE_NODISCARD bool FNoDiscardAsserter::IsNull(const T& Ptr, const FString& FailureMessage)
 {
 	if (Ptr != nullptr)
 	{
@@ -99,25 +99,25 @@ template <typename T>
 }
 
 template <typename T>
-[[nodiscard]] bool FNoDiscardAsserter::IsNotNull(const T& Ptr)
+UE_NODISCARD bool FNoDiscardAsserter::IsNotNull(const T& Ptr)
 {
 	return IsNotNull(Ptr, "TestNotNull Failed");
 }
 
 template <typename T>
-[[nodiscard]] bool FNoDiscardAsserter::IsNotNull(const T& Ptr, const char* FailureMessage)
+UE_NODISCARD bool FNoDiscardAsserter::IsNotNull(const T& Ptr, const char* FailureMessage)
 {
 	return IsNotNull(Ptr, FString(FailureMessage));
 }
 
 template <typename T>
-[[nodiscard]] bool FNoDiscardAsserter::IsNotNull(const T& Ptr, const TCHAR* FailureMessage)
+UE_NODISCARD bool FNoDiscardAsserter::IsNotNull(const T& Ptr, const TCHAR* FailureMessage)
 {
 	return IsNotNull(Ptr, FString(FailureMessage));
 }
 
 template <typename T>
-[[nodiscard]] bool FNoDiscardAsserter::IsNotNull(const T& Ptr, const FString& FailureMessage)
+UE_NODISCARD bool FNoDiscardAsserter::IsNotNull(const T& Ptr, const FString& FailureMessage)
 {
 	if (Ptr == nullptr)
 	{
@@ -129,76 +129,76 @@ template <typename T>
 
 
 template <typename TExpected, typename TActual>
-[[nodiscard]] bool FNoDiscardAsserter::AreEqual(const TExpected& Expected, const TActual& Actual)
+UE_NODISCARD bool FNoDiscardAsserter::AreEqual(const TExpected& Expected, const TActual& Actual)
 {
 	FString Message = FString::Printf(TEXT("Expected %s to equal %s"), *CQTestConvert::ToString(Expected), *CQTestConvert::ToString(Actual));
 	return AreEqual(Expected, Actual, Message);
 }
 
 template <typename TExpected, typename TActual>
-[[nodiscard]] bool FNoDiscardAsserter::AreEqual(const TExpected& Expected, const TActual& Actual, const char* FailureMessage)
+UE_NODISCARD bool FNoDiscardAsserter::AreEqual(const TExpected& Expected, const TActual& Actual, const char* FailureMessage)
 {
 	return AreEqual(Expected, Actual, FString(FailureMessage));
 }
 
 template <typename TExpected, typename TActual>
-[[nodiscard]] bool FNoDiscardAsserter::AreEqual(const TExpected& Expected, const TActual& Actual, const TCHAR* FailureMessage)
+UE_NODISCARD bool FNoDiscardAsserter::AreEqual(const TExpected& Expected, const TActual& Actual, const TCHAR* FailureMessage)
 {
 	return AreEqual(Expected, Actual, FString(FailureMessage));
 }
 
 template <typename TExpected, typename TActual>
-[[nodiscard]] bool FNoDiscardAsserter::AreEqual(const TExpected& Expected, const TActual& Actual, const FString& FailureMessage)
+UE_NODISCARD bool FNoDiscardAsserter::AreEqual(const TExpected& Expected, const TActual& Actual, const FString& FailureMessage)
 {
 	return TestRunner.TestEqual(FailureMessage, Actual, Expected); // FAutomationTestBase expects Actual then Expected
 }
 
 template <typename TExpected, typename TActual>
-[[nodiscard]] bool FNoDiscardAsserter::AreNotEqual(const TExpected& Expected, const TActual& Actual)
+UE_NODISCARD bool FNoDiscardAsserter::AreNotEqual(const TExpected& Expected, const TActual& Actual)
 {
 	FString Message = FString::Printf(TEXT("Expected %s to not equal %s"), *CQTestConvert::ToString(Expected), *CQTestConvert::ToString(Actual));
 	return AreNotEqual(Expected, Actual, Message);
 }
 
 template <typename TExpected, typename TActual>
-[[nodiscard]] bool FNoDiscardAsserter::AreNotEqual(const TExpected& Expected, const TActual& Actual, const char* FailureMessage)
+UE_NODISCARD bool FNoDiscardAsserter::AreNotEqual(const TExpected& Expected, const TActual& Actual, const char* FailureMessage)
 {
 	return AreNotEqual(Expected, Actual, FString(FailureMessage));
 }
 
 template <typename TExpected, typename TActual>
-[[nodiscard]] bool FNoDiscardAsserter::AreNotEqual(const TExpected& Expected, const TActual& Actual, const TCHAR* FailureMessage)
+UE_NODISCARD bool FNoDiscardAsserter::AreNotEqual(const TExpected& Expected, const TActual& Actual, const TCHAR* FailureMessage)
 {
 	return AreNotEqual(Expected, Actual, FString(FailureMessage));
 }
 
 template <typename TExpected, typename TActual>
-[[nodiscard]] bool FNoDiscardAsserter::AreNotEqual(const TExpected& Expected, const TActual& Actual, const FString& FailureMessage)
+UE_NODISCARD bool FNoDiscardAsserter::AreNotEqual(const TExpected& Expected, const TActual& Actual, const FString& FailureMessage)
 {
 	return TestRunner.TestNotEqual(FailureMessage, Actual, Expected); // FAutomationTestBase expects Actual then Expected
 }
 
 template <typename TExpected, typename TActual, typename TEpsilon>
-[[nodiscard]] bool FNoDiscardAsserter::IsNear(TExpected Expected, TActual Actual, TEpsilon Epsilon)
+UE_NODISCARD bool FNoDiscardAsserter::IsNear(TExpected Expected, TActual Actual, TEpsilon Epsilon)
 {
 	FString Message = FString::Printf(TEXT("Expected %s to be near %s (within %s)"), *CQTestConvert::ToString(Expected), *CQTestConvert::ToString(Actual), *CQTestConvert::ToString(Epsilon));
 	return IsNear(Expected, Actual, Epsilon, Message);
 }
 
 template <typename TExpected, typename TActual, typename TEpsilon>
-[[nodiscard]] bool FNoDiscardAsserter::IsNear(TExpected Expected, TActual Actual, TEpsilon Epsilon, const char* FailureMessage)
+UE_NODISCARD bool FNoDiscardAsserter::IsNear(TExpected Expected, TActual Actual, TEpsilon Epsilon, const char* FailureMessage)
 {
 	return IsNear(Expected, Actual, Epsilon, FString(FailureMessage));
 }
 
 template <typename TExpected, typename TActual, typename TEpsilon>
-[[nodiscard]] inline bool FNoDiscardAsserter::IsNear(TExpected Expected, TActual Actual, TEpsilon Epsilon, const TCHAR* FailureMessage)
+UE_NODISCARD inline bool FNoDiscardAsserter::IsNear(TExpected Expected, TActual Actual, TEpsilon Epsilon, const TCHAR* FailureMessage)
 {
 	return IsNear(Expected, Actual, Epsilon, FString(FailureMessage));
 }
 
 template <typename TExpected, typename TActual, typename TEpsilon>
-[[nodiscard]] inline bool FNoDiscardAsserter::IsNear(TExpected Expected, TActual Actual, TEpsilon Epsilon, const FString& FailureMessage)
+UE_NODISCARD inline bool FNoDiscardAsserter::IsNear(TExpected Expected, TActual Actual, TEpsilon Epsilon, const FString& FailureMessage)
 {
 	return TestRunner.TestEqual(FailureMessage, Actual, Expected, Epsilon); // FAutomationTestBase expects Actual then Expected
 }
