[CQTest] Implement assert test conditions instead of using the one from FAutomationTestBase to control messaging and interface.
+ fix enum to string handling

#rnx

diff --git a/Engine/Plugins/Tests/CQTest/Source/CQTest/Public/Assert/CQTestCondition.h b/Engine/Plugins/Tests/CQTest/Source/CQTest/Public/Assert/CQTestCondition.h
new file mode 100644
--- /dev/null
+++ b/Engine/Plugins/Tests/CQTest/Source/CQTest/Public/Assert/CQTestCondition.h
@@ -0,0 +1,136 @@
+// Copyright Epic Games, Inc. All Rights Reserved.
+#pragma once
+
+#include "Containers/UnrealString.h"
+#include "UObject/NameTypes.h"
+
+namespace CQTestCondition
+{
+	// Equal
+
+	template <typename TValue, typename TOtherValue>
+	inline bool IsEqual(const TValue& A, const TOtherValue& B)
+	{
+		return A == B;
+	}
+
+	inline bool IsEqual(const TCHAR* A, const TCHAR* B)
+	{
+		return FCString::Strcmp(A, B) == 0;
+	}
+	inline bool IsEqual(const TCHAR* A, const FString& B)
+	{
+		return IsEqual(A, *B);
+	}
+	inline bool IsEqual(const FString& A, const TCHAR* B)
+	{
+		return IsEqual(*A, B);
+	}
+	inline bool IsEqual(const FString& A, const FString& B)
+	{
+		return IsEqual(*A, *B);
+	}
+
+	inline bool IsEqualIgnoreCase(const TCHAR* A, const TCHAR* B)
+	{
+		return FCString::Stricmp(A, B) == 0;
+	}
+	inline bool IsEqualIgnoreCase(const TCHAR* A, const FString& B)
+	{
+		return IsEqualIgnoreCase(A, *B);
+	}
+	inline bool IsEqualIgnoreCase(const FString& A, const TCHAR* B)
+	{
+		return IsEqualIgnoreCase(*A, B);
+	}
+	inline bool IsEqualIgnoreCase(const FString& A, const FString& B)
+	{
+		return IsEqualIgnoreCase(*A, *B);
+	}
+
+	// Not Equal
+
+	template <typename TValue, typename TOtherValue>
+	inline bool IsNotEqual(const TValue& A, const TOtherValue& B)
+	{
+		return A != B;
+	}
+
+	inline bool IsNotEqual(const TCHAR* A, const TCHAR* B)
+	{
+		return FCString::Strcmp(A, B) != 0;
+	}
+	inline bool IsNotEqual(const TCHAR* A, const FString& B)
+	{
+		return IsNotEqual(A, *B);
+	}
+	inline bool IsNotEqual(const FString& A, const TCHAR* B)
+	{
+		return IsNotEqual(*A, B);
+	}
+	inline bool IsNotEqual(const FString& A, const FString& B)
+	{
+		return IsNotEqual(*A, *B);
+	}
+
+	inline bool IsNotEqualIgnoreCase(const TCHAR* A, const TCHAR* B)
+	{
+		return FCString::Stricmp(A, B) != 0;
+	}
+	inline bool IsNotEqualIgnoreCase(const TCHAR* A, const FString& B)
+	{
+		return IsNotEqualIgnoreCase(A, *B);
+	}
+	inline bool IsNotEqualIgnoreCase(const FString& A, const TCHAR* B)
+	{
+		return IsNotEqualIgnoreCase(*A, B);
+	}
+	inline bool IsNotEqualIgnoreCase(const FString& A, const FString& B)
+	{
+		return IsNotEqualIgnoreCase(*A, *B);
+	}
+
+	// Approximation
+
+	inline bool IsNearlyEqual(const float A, const float B, float Tolerance = UE_KINDA_SMALL_NUMBER)
+	{
+		return FMath::IsNearlyEqual(A, B, Tolerance);
+	}
+	inline bool IsNearlyEqual(const double A, const double B, double Tolerance = UE_KINDA_SMALL_NUMBER)
+	{
+		return FMath::IsNearlyEqual(A, B, Tolerance);
+	}
+	inline bool IsNearlyEqual(const FVector A, const FVector B, float Tolerance = UE_KINDA_SMALL_NUMBER)
+	{
+		return A.Equals(B, Tolerance);
+	}
+	inline bool IsNearlyEqual(const FTransform A, const FTransform B, float Tolerance = UE_KINDA_SMALL_NUMBER)
+	{
+		return A.Equals(B, Tolerance);
+	}
+	inline bool IsNearlyEqual(const FRotator A, const FRotator B, float Tolerance = UE_KINDA_SMALL_NUMBER)
+	{
+		return A.Equals(B, Tolerance);
+	}
+
+	inline bool IsNotNearlyEqual(const float A, const float B, float Tolerance = UE_KINDA_SMALL_NUMBER)
+	{
+		return !FMath::IsNearlyEqual(A, B, Tolerance);
+	}
+	inline bool IsNotNearlyEqual(const double A, const double B, double Tolerance = UE_KINDA_SMALL_NUMBER)
+	{
+		return !FMath::IsNearlyEqual(A, B, Tolerance);
+	}
+	inline bool IsNotNearlyEqual(const FVector A, const FVector B, float Tolerance = UE_KINDA_SMALL_NUMBER)
+	{
+		return !A.Equals(B, Tolerance);
+	}
+	inline bool IsNotNearlyEqual(const FTransform A, const FTransform B, float Tolerance = UE_KINDA_SMALL_NUMBER)
+	{
+		return !A.Equals(B, Tolerance);
+	}
+	inline bool IsNotNearlyEqual(const FRotator A, const FRotator B, float Tolerance = UE_KINDA_SMALL_NUMBER)
+	{
+		return !A.Equals(B, Tolerance);
+	}
+}
diff --git a/Engine/Plugins/Tests/CQTest/Source/CQTest/Public/Assert/CQTestConvert.h b/Engine/Plugins/Tests/CQTest/Source/CQTest/Public/Assert/CQTestConvert.h
--- a/Engine/Plugins/Tests/CQTest/Source/CQTest/Public/Assert/CQTestConvert.h
+++ b/Engine/Plugins/Tests/CQTest/Source/CQTest/Public/Assert/CQTestConvert.h
@@ -2,13 +2,14 @@
 #pragma once
 
 #include "Containers/UnrealString.h"
+#include "Misc/StringBuilder.h"
 #include "UObject/NameTypes.h"
 #include <sstream>
 #include <type_traits>
 
 namespace CQTestConvert
 {
-	using PrintStream = std::ostringstream;
+	using PrintStream = FStringBuilderBase;
 }
 
 namespace
@@ -27,10 +28,8 @@ namespace
 	template<typename T>
 	class THasOStream
 	{
-		template <typename C, typename = decltype(std::declval<CQTestConvert::PrintStream&>().operator<<(std::declval<C>()))>
-		static std::true_type test(int);
 		template <typename C, typename = decltype(operator<<(std::declval<CQTestConvert::PrintStream&>(), std::declval<C>()))>
-		static std::true_type test(char);
+		static std::true_type test(int);
 		template<typename C>
 		static std::false_type test(...);
 
@@ -88,19 +87,17 @@ namespace CQTestConvert
 		{
 			PrintStream stream;
 			stream << Input;
-			return FString(stream.str().c_str(), stream.str().length());
+			return stream.ToString();
+		}
+		else if constexpr (std::is_floating_point_v<T>)
+		{
+			// Fallback for floating point type
+			return FString::Printf(TEXT("%f"), Input);
 		}
 		else if constexpr (std::is_enum_v<T>)
 		{
-			// Check special case of uint8/int8 to avoid interpretation as a char
-			if constexpr (std::is_same_v<uint8, std::underlying_type_t<T>> || std::is_same_v<int8, std::underlying_type_t<T>>)
-			{
-				return ToString(static_cast<int>(Input));
-			}
-			else
-			{
-				return ToString(static_cast<std::underlying_type_t<T>>(Input));
-			}
+			// Fallback for enum type
+			return ToString(static_cast<std::underlying_type_t<T>>(Input));
 		}
 		else
 		{
diff --git a/Engine/Plugins/Tests/CQTest/Source/CQTest/Public/Assert/NoDiscardAsserter.h b/Engine/Plugins/Tests/CQTest/Source/CQTest/Public/Assert/NoDiscardAsserter.h
--- a/Engine/Plugins/Tests/CQTest/Source/CQTest/Public/Assert/NoDiscardAsserter.h
+++ b/Engine/Plugins/Tests/CQTest/Source/CQTest/Public/Assert/NoDiscardAsserter.h
@@ -64,6 +64,14 @@ public:
 	template <typename TExpected, typename TActual>
 	UE_NODISCARD bool AreEqual(const TExpected& Expected, const TActual& Actual, const FString& FailureMessage);
 
+	UE_NODISCARD bool AreEqualIgnoreCase(const FString& Expected, const FString& Actual);
+
+	UE_NODISCARD bool AreEqualIgnoreCase(const FString& Expected, const FString& Actual, const char* FailureMessage);
+
+	UE_NODISCARD bool AreEqualIgnoreCase(const FString& Expected, const FString& Actual, const TCHAR* FailureMessage);
+
+	UE_NODISCARD bool AreEqualIgnoreCase(const FString& Expected, const FString& Actual, const FString& FailureMessage);
+
 	template <typename TExpected, typename TActual>
 	UE_NODISCARD bool AreNotEqual(const TExpected& Expected, const TActual& Actual);
 
@@ -76,6 +84,14 @@ public:
 	template <typename TExpected, typename TActual>
 	UE_NODISCARD bool AreNotEqual(const TExpected& Expected, const TActual& Actual, const FString& FailureMessage);
 
+	UE_NODISCARD bool AreNotEqualIgnoreCase(const FString& Expected, const FString& Actual);
+
+	UE_NODISCARD bool AreNotEqualIgnoreCase(const FString& Expected, const FString& Actual, const char* FailureMessage);
+
+	UE_NODISCARD bool AreNotEqualIgnoreCase(const FString& Expected, const FString& Actual, const TCHAR* FailureMessage);
+
+	UE_NODISCARD bool AreNotEqualIgnoreCase(const FString& Expected, const FString& Actual, const FString& FailureMessage);
+
 	template <typename TExpected, typename TActual, typename TEpsilon>
 	UE_NODISCARD bool IsNear(TExpected Expected, TActual Actual, TEpsilon Epsilon);
 
@@ -87,6 +103,7 @@ public:
 
 	template <typename TExpected, typename TActual, typename TEpsilon>
 	UE_NODISCARD bool IsNear(TExpected Expected, TActual Actual, TEpsilon Epsilon, const FString& FailureMessage);
+
 };
 
 #include "Assert/NoDiscardAsserter.inl"
\ No newline at end of file
diff --git a/Engine/Plugins/Tests/CQTest/Source/CQTest/Public/Assert/NoDiscardAsserter.inl b/Engine/Plugins/Tests/CQTest/Source/CQTest/Public/Assert/NoDiscardAsserter.inl
--- a/Engine/Plugins/Tests/CQTest/Source/CQTest/Public/Assert/NoDiscardAsserter.inl
+++ b/Engine/Plugins/Tests/CQTest/Source/CQTest/Public/Assert/NoDiscardAsserter.inl
@@ -2,6 +2,7 @@
 
 #include "Internationalization/Regex.h"
 #include "Assert/CQTestConvert.h"
+#include "Assert/CQTestCondition.h"
 
 inline void FNoDiscardAsserter::ExpectError(FString Error, int32 Count)
 {
@@ -24,71 +25,81 @@ inline void FNoDiscardAsserter::ExpectErrorRegex(FString Error, int32 Count)
 	TestRunner.AddExpectedError(Error, EAutomationExpectedErrorFlags::Contains, Count);
 }
 
-inline void FNoDiscardAsserter::Fail(FString Error)
+FORCEINLINE void FNoDiscardAsserter::Fail(FString Error)
 {
 	TestRunner.AddError(Error);
 }
 
-UE_NODISCARD inline bool FNoDiscardAsserter::IsTrue(bool Condition)
+UE_NODISCARD FORCEINLINE bool FNoDiscardAsserter::IsTrue(bool Condition)
 {
-	return IsTrue(Condition, "condition");
+	return IsTrue(Condition, TEXT("Expected condition to be true."));
 }
 
-UE_NODISCARD inline bool FNoDiscardAsserter::IsTrue(bool Condition, const char* FailureMessage)
+UE_NODISCARD FORCEINLINE bool FNoDiscardAsserter::IsTrue(bool Condition, const char* FailureMessage)
 {
 	return IsTrue(Condition, FString(FailureMessage));
 }
 
-UE_NODISCARD inline bool FNoDiscardAsserter::IsTrue(bool Condition, const TCHAR* FailureMessage)
+UE_NODISCARD FORCEINLINE bool FNoDiscardAsserter::IsTrue(bool Condition, const TCHAR* FailureMessage)
 {
 	return IsTrue(Condition, FString(FailureMessage));
 }
 
-UE_NODISCARD inline bool FNoDiscardAsserter::IsTrue(bool Condition, const FString& FailureMessage)
+UE_NODISCARD FORCEINLINE bool FNoDiscardAsserter::IsTrue(bool Condition, const FString& FailureMessage)
 {
-	return TestRunner.TestTrue(FailureMessage, Condition);
+	if (!Condition)
+	{
+		Fail(FailureMessage);
+		return false;
+	}
+	return true;
 }
 
-UE_NODISCARD inline bool FNoDiscardAsserter::IsFalse(bool Condition)
+UE_NODISCARD FORCEINLINE bool FNoDiscardAsserter::IsFalse(bool Condition)
 {
-	return IsFalse(Condition, "condition");
+	return IsFalse(Condition, TEXT("Expected condition to be false."));
 }
 
-UE_NODISCARD inline bool FNoDiscardAsserter::IsFalse(bool Condition, const char* FailureMessage)
+UE_NODISCARD FORCEINLINE bool FNoDiscardAsserter::IsFalse(bool Condition, const char* FailureMessage)
 {
 	return IsFalse(Condition, FString(FailureMessage));
 }
 
-UE_NODISCARD inline bool FNoDiscardAsserter::IsFalse(bool Condition, const TCHAR* FailureMessage)
+UE_NODISCARD FORCEINLINE bool FNoDiscardAsserter::IsFalse(bool Condition, const TCHAR* FailureMessage)
 {
 	return IsFalse(Condition, FString(FailureMessage));
 }
 
-UE_NODISCARD inline bool FNoDiscardAsserter::IsFalse(bool Condition, const FString& FailureMessage)
+UE_NODISCARD FORCEINLINE bool FNoDiscardAsserter::IsFalse(bool Condition, const FString& FailureMessage)
 {
-	return TestRunner.TestFalse(FString(FailureMessage), Condition);
+	if (Condition)
+	{
+		Fail(FailureMessage);
+		return false;
+	}
+	return true;
 }
 
 template <typename T>
-UE_NODISCARD bool FNoDiscardAsserter::IsNull(const T& Ptr)
+UE_NODISCARD FORCEINLINE bool FNoDiscardAsserter::IsNull(const T& Ptr)
 {
-	return IsNull(Ptr, "TestNull failed");
+	return IsNull(Ptr, TEXT("Expected pointer to be Null"));
 }
 
 template <typename T>
-UE_NODISCARD bool FNoDiscardAsserter::IsNull(const T& Ptr, const char* FailureMessage)
+UE_NODISCARD FORCEINLINE bool FNoDiscardAsserter::IsNull(const T& Ptr, const char* FailureMessage)
 {
 	return IsNull(Ptr, FString(FailureMessage));
 }
 
 template <typename T>
-UE_NODISCARD bool FNoDiscardAsserter::IsNull(const T& Ptr, const TCHAR* FailureMessage)
+UE_NODISCARD FORCEINLINE bool FNoDiscardAsserter::IsNull(const T& Ptr, const TCHAR* FailureMessage)
 {
 	return IsNull(Ptr, FString(FailureMessage));
 }
 
 template <typename T>
-UE_NODISCARD bool FNoDiscardAsserter::IsNull(const T& Ptr, const FString& FailureMessage)
+UE_NODISCARD FORCEINLINE bool FNoDiscardAsserter::IsNull(const T& Ptr, const FString& FailureMessage)
 {
 	if (Ptr != nullptr)
 	{
@@ -99,25 +110,25 @@ UE_NODISCARD bool FNoDiscardAsserter::IsNull(const T& Ptr, const FString& Failur
 }
 
 template <typename T>
-UE_NODISCARD bool FNoDiscardAsserter::IsNotNull(const T& Ptr)
+UE_NODISCARD FORCEINLINE bool FNoDiscardAsserter::IsNotNull(const T& Ptr)
 {
-	return IsNotNull(Ptr, "TestNotNull Failed");
+	return IsNotNull(Ptr, TEXT("Expected pointer to be not Null"));
 }
 
 template <typename T>
-UE_NODISCARD bool FNoDiscardAsserter::IsNotNull(const T& Ptr, const char* FailureMessage)
+UE_NODISCARD FORCEINLINE bool FNoDiscardAsserter::IsNotNull(const T& Ptr, const char* FailureMessage)
 {
 	return IsNotNull(Ptr, FString(FailureMessage));
 }
 
 template <typename T>
-UE_NODISCARD bool FNoDiscardAsserter::IsNotNull(const T& Ptr, const TCHAR* FailureMessage)
+UE_NODISCARD FORCEINLINE bool FNoDiscardAsserter::IsNotNull(const T& Ptr, const TCHAR* FailureMessage)
 {
 	return IsNotNull(Ptr, FString(FailureMessage));
 }
 
 template <typename T>
-UE_NODISCARD bool FNoDiscardAsserter::IsNotNull(const T& Ptr, const FString& FailureMessage)
+UE_NODISCARD FORCEINLINE bool FNoDiscardAsserter::IsNotNull(const T& Ptr, const FString& FailureMessage)
 {
 	if (Ptr == nullptr)
 	{
@@ -127,78 +138,164 @@ UE_NODISCARD bool FNoDiscardAsserter::IsNotNull(const T& Ptr, const FString& Fai
 	return true;
 }
 
-
 template <typename TExpected, typename TActual>
-UE_NODISCARD bool FNoDiscardAsserter::AreEqual(const TExpected& Expected, const TActual& Actual)
+UE_NODISCARD FORCEINLINE bool FNoDiscardAsserter::AreEqual(const TExpected& Expected, const TActual& Actual)
 {
-	FString Message = FString::Printf(TEXT("Expected %s to equal %s"), *CQTestConvert::ToString(Expected), *CQTestConvert::ToString(Actual));
-	return AreEqual(Expected, Actual, Message);
+	if (CQTestCondition::IsNotEqual(Expected, Actual))
+	{
+		Fail(FString::Printf(TEXT("Expected %s to equal %s."), *CQTestConvert::ToString(Expected), *CQTestConvert::ToString(Actual)));
+		return false;
+	}
+	return true;
 }
 
 template <typename TExpected, typename TActual>
-UE_NODISCARD bool FNoDiscardAsserter::AreEqual(const TExpected& Expected, const TActual& Actual, const char* FailureMessage)
+UE_NODISCARD FORCEINLINE bool FNoDiscardAsserter::AreEqual(const TExpected& Expected, const TActual& Actual, const char* FailureMessage)
 {
 	return AreEqual(Expected, Actual, FString(FailureMessage));
 }
 
 template <typename TExpected, typename TActual>
-UE_NODISCARD bool FNoDiscardAsserter::AreEqual(const TExpected& Expected, const TActual& Actual, const TCHAR* FailureMessage)
+UE_NODISCARD FORCEINLINE bool FNoDiscardAsserter::AreEqual(const TExpected& Expected, const TActual& Actual, const TCHAR* FailureMessage)
 {
 	return AreEqual(Expected, Actual, FString(FailureMessage));
 }
 
 template <typename TExpected, typename TActual>
-UE_NODISCARD bool FNoDiscardAsserter::AreEqual(const TExpected& Expected, const TActual& Actual, const FString& FailureMessage)
+UE_NODISCARD FORCEINLINE bool FNoDiscardAsserter::AreEqual(const TExpected& Expected, const TActual& Actual, const FString& FailureMessage)
+{
+	if (CQTestCondition::IsNotEqual(Expected, Actual))
+	{
+		Fail(FailureMessage);
+		return false;
+	}
+	return true;
+}
+
+UE_NODISCARD FORCEINLINE bool FNoDiscardAsserter::AreEqualIgnoreCase(const FString& Expected, const FString& Actual)
+{
+	if (CQTestCondition::IsNotEqualIgnoreCase(Expected, Actual))
+	{
+		Fail(FString::Printf(TEXT("Expected %s to equal %s ignoring case."), *Expected, *Actual));
+		return false;
+	}
+	return true;
+}
+
+UE_NODISCARD FORCEINLINE bool FNoDiscardAsserter::AreEqualIgnoreCase(const FString& Expected, const FString& Actual, const char* FailureMessage)
 {
-	return TestRunner.TestEqual(FailureMessage, Actual, Expected); // FAutomationTestBase expects Actual then Expected
+	return AreEqualIgnoreCase(Expected, Actual, FString(FailureMessage));
+}
+
+UE_NODISCARD FORCEINLINE bool FNoDiscardAsserter::AreEqualIgnoreCase(const FString& Expected, const FString& Actual, const TCHAR* FailureMessage)
+{
+	return AreEqualIgnoreCase(Expected, Actual, FString(FailureMessage));
+}
+
+UE_NODISCARD FORCEINLINE bool FNoDiscardAsserter::AreEqualIgnoreCase(const FString& Expected, const FString& Actual, const FString& FailureMessage)
+{
+	if (CQTestCondition::IsNotEqualIgnoreCase(Expected, Actual))
+	{
+		Fail(FailureMessage);
+		return false;
+	}
+	return true;
 }
 
 template <typename TExpected, typename TActual>
-UE_NODISCARD bool FNoDiscardAsserter::AreNotEqual(const TExpected& Expected, const TActual& Actual)
+UE_NODISCARD FORCEINLINE bool FNoDiscardAsserter::AreNotEqual(const TExpected& Expected, const TActual& Actual)
 {
-	FString Message = FString::Printf(TEXT("Expected %s to not equal %s"), *CQTestConvert::ToString(Expected), *CQTestConvert::ToString(Actual));
-	return AreNotEqual(Expected, Actual, Message);
+	if (CQTestCondition::IsEqual(Expected, Actual))
+	{
+		Fail(FString::Printf(TEXT("Expected %s to not equal %s."), *CQTestConvert::ToString(Expected), *CQTestConvert::ToString(Actual)));
+		return false;
+	}
+	return true;
 }
 
 template <typename TExpected, typename TActual>
-UE_NODISCARD bool FNoDiscardAsserter::AreNotEqual(const TExpected& Expected, const TActual& Actual, const char* FailureMessage)
+UE_NODISCARD FORCEINLINE bool FNoDiscardAsserter::AreNotEqual(const TExpected& Expected, const TActual& Actual, const char* FailureMessage)
 {
 	return AreNotEqual(Expected, Actual, FString(FailureMessage));
 }
 
 template <typename TExpected, typename TActual>
-UE_NODISCARD bool FNoDiscardAsserter::AreNotEqual(const TExpected& Expected, const TActual& Actual, const TCHAR* FailureMessage)
+UE_NODISCARD FORCEINLINE bool FNoDiscardAsserter::AreNotEqual(const TExpected& Expected, const TActual& Actual, const TCHAR* FailureMessage)
 {
 	return AreNotEqual(Expected, Actual, FString(FailureMessage));
 }
 
 template <typename TExpected, typename TActual>
-UE_NODISCARD bool FNoDiscardAsserter::AreNotEqual(const TExpected& Expected, const TActual& Actual, const FString& FailureMessage)
+UE_NODISCARD FORCEINLINE bool FNoDiscardAsserter::AreNotEqual(const TExpected& Expected, const TActual& Actual, const FString& FailureMessage)
+{
+	if (CQTestCondition::IsEqual(Expected, Actual))
+	{
+		Fail(FailureMessage);
+		return false;
+	}
+	return true;
+}
+
+UE_NODISCARD FORCEINLINE bool FNoDiscardAsserter::AreNotEqualIgnoreCase(const FString& Expected, const FString& Actual)
+{
+	if (CQTestCondition::IsEqualIgnoreCase(Expected, Actual))
+	{
+		Fail(FString::Printf(TEXT("Expected %s to not equal %s ignoring case."), *Expected, *Actual));
+		return false;
+	}
+	return true;
+}
+
+UE_NODISCARD FORCEINLINE bool FNoDiscardAsserter::AreNotEqualIgnoreCase(const FString& Expected, const FString& Actual, const char* FailureMessage)
 {
-	return TestRunner.TestNotEqual(FailureMessage, Actual, Expected); // FAutomationTestBase expects Actual then Expected
+	return AreNotEqualIgnoreCase(Expected, Actual, FString(FailureMessage));
+}
+
+UE_NODISCARD FORCEINLINE bool FNoDiscardAsserter::AreNotEqualIgnoreCase(const FString& Expected, const FString& Actual, const TCHAR* FailureMessage)
+{
+	return AreNotEqualIgnoreCase(Expected, Actual, FString(FailureMessage));
+}
+
+UE_NODISCARD FORCEINLINE bool FNoDiscardAsserter::AreNotEqualIgnoreCase(const FString& Expected, const FString& Actual, const FString& FailureMessage)
+{
+	if (CQTestCondition::IsEqualIgnoreCase(Expected, Actual))
+	{
+		Fail(FailureMessage);
+		return false;
+	}
+	return true;
 }
 
 template <typename TExpected, typename TActual, typename TEpsilon>
-UE_NODISCARD bool FNoDiscardAsserter::IsNear(TExpected Expected, TActual Actual, TEpsilon Epsilon)
+UE_NODISCARD FORCEINLINE bool FNoDiscardAsserter::IsNear(TExpected Expected, TActual Actual, TEpsilon Epsilon)
 {
-	FString Message = FString::Printf(TEXT("Expected %s to be near %s (within %s)"), *CQTestConvert::ToString(Expected), *CQTestConvert::ToString(Actual), *CQTestConvert::ToString(Epsilon));
-	return IsNear(Expected, Actual, Epsilon, Message);
+	if (CQTestCondition::IsNotNearlyEqual(Expected, Actual, Epsilon))
+	{
+		Fail(FString::Printf(TEXT("Expected %s to be near %s (within %s)."), *CQTestConvert::ToString(Expected), *CQTestConvert::ToString(Actual), *CQTestConvert::ToString(Epsilon)));
+		return false;
+	}
+	return true;
 }
 
 template <typename TExpected, typename TActual, typename TEpsilon>
-UE_NODISCARD bool FNoDiscardAsserter::IsNear(TExpected Expected, TActual Actual, TEpsilon Epsilon, const char* FailureMessage)
+UE_NODISCARD FORCEINLINE bool FNoDiscardAsserter::IsNear(TExpected Expected, TActual Actual, TEpsilon Epsilon, const char* FailureMessage)
 {
 	return IsNear(Expected, Actual, Epsilon, FString(FailureMessage));
 }
 
 template <typename TExpected, typename TActual, typename TEpsilon>
-UE_NODISCARD inline bool FNoDiscardAsserter::IsNear(TExpected Expected, TActual Actual, TEpsilon Epsilon, const TCHAR* FailureMessage)
+UE_NODISCARD FORCEINLINE bool FNoDiscardAsserter::IsNear(TExpected Expected, TActual Actual, TEpsilon Epsilon, const TCHAR* FailureMessage)
 {
 	return IsNear(Expected, Actual, Epsilon, FString(FailureMessage));
 }
 
 template <typename TExpected, typename TActual, typename TEpsilon>
-UE_NODISCARD inline bool FNoDiscardAsserter::IsNear(TExpected Expected, TActual Actual, TEpsilon Epsilon, const FString& FailureMessage)
+UE_NODISCARD FORCEINLINE bool FNoDiscardAsserter::IsNear(TExpected Expected, TActual Actual, TEpsilon Epsilon, const FString& FailureMessage)
 {
-	return TestRunner.TestEqual(FailureMessage, Actual, Expected, Epsilon); // FAutomationTestBase expects Actual then Expected
+	if (CQTestCondition::IsNotNearlyEqual(Expected, Actual, Epsilon))
+	{
+		Fail(FailureMessage);
+		return false;
+	}
+	return true;
 }
diff --git a/Engine/Plugins/Tests/CQTest/Source/CQTestTests/Private/Assert/AssertionTests.cpp b/Engine/Plugins/Tests/CQTest/Source/CQTestTests/Private/Assert/AssertionTests.cpp
--- a/Engine/Plugins/Tests/CQTest/Source/CQTestTests/Private/Assert/AssertionTests.cpp
+++ b/Engine/Plugins/Tests/CQTest/Source/CQTestTests/Private/Assert/AssertionTests.cpp
@@ -7,14 +7,16 @@ DECLARE_LOG_CATEGORY_CLASS(AssertionTests, Log, All);
 
 namespace CQTests
 {
-	TEST_CLASS(NoDiscardAssert, "TestFramework.CQTest.Core")
+	static const FString ExpectedError = TEXT("Expected");
+	static const TCHAR* AnyError = TEXT("");
+
+	TEST_CLASS(NoDiscardAssert_Errors, "TestFramework.CQTest.Core")
 	{
-		const char* ExpectedError = "Expected";
 
 		TEST_METHOD(AssertFail_WithMessage_AddsError)
 		{
-			Assert.ExpectError(TEXT(""));
-			Assert.Fail(TEXT(""));
+			Assert.ExpectError(AnyError);
+			Assert.Fail(AnyError);
 		}
 
 		TEST_METHOD(AssertFail_Macro_AddsErrorAndExits)
@@ -76,7 +78,10 @@ namespace CQTests
 			Assert.ExpectErrorRegex("\\w+");
 			Assert.Fail("abc");
 		}
+	};
 
+	TEST_CLASS(NoDiscardAssert_Bools, "TestFramework.CQTest.Core")
+	{
 		TEST_METHOD(AssertTrue_WithTrue_Succeeds)
 		{
 			ASSERT_THAT(IsTrue(true));
@@ -87,52 +92,54 @@ namespace CQTests
 		}
 		TEST_METHOD(AssertTrue_WithFalse_AddsError)
 		{
-			Assert.ExpectError(TEXT(""));
+			Assert.ExpectError(AnyError);
 			ASSERT_THAT(IsTrue(false));
 		}
-		TEST_METHOD(AssertTrue_WithFalseAndError_AddSpecificError)
+		TEST_METHOD(AssertTrue_WithFalseAndError_AddsSpecificError)
 		{
-			Assert.ExpectError(FString(ExpectedError));
+			Assert.ExpectError(ExpectedError);
 			ASSERT_THAT(IsTrue(false, ExpectedError));
 		}
 
-		TEST_METHOD(AssertEqual_WithSameInts_Succeeds)
+		TEST_METHOD(AssertFalse_WithFalse_Succeeds)
 		{
-			ASSERT_THAT(AreEqual(42, 42));
+			ASSERT_THAT(IsFalse(false));
 		}
-		TEST_METHOD(AssertEqual_WithSameIntsAndErrorMessage_DoesNotAddErrorMessage)
+		TEST_METHOD(AssertFalse_WithFalseAndErrorMessage_DoesNotAddErrorMessage)
 		{
-			ASSERT_THAT(AreEqual(42, 42, "Unexpected"));
+			ASSERT_THAT(IsFalse(false, "Unexpected"));
 		}
-		TEST_METHOD(AssertEqual_WithDifferentInts_AddsError)
+		TEST_METHOD(AssertFalse_WithTrue_AddsError)
 		{
-			Assert.ExpectError(TEXT(""));
-			ASSERT_THAT(AreEqual(42, 0));
+			Assert.ExpectError(AnyError);
+			ASSERT_THAT(IsFalse(true));
 		}
-		TEST_METHOD(AssertEqual_WithDifferentIntsAndErrorMessage_AddsSpecificError)
+		TEST_METHOD(AssertFalse_WithTrueAndError_AddsSpecificError)
 		{
-			Assert.ExpectError(FString(ExpectedError));
-			ASSERT_THAT(AreEqual(42, 0, ExpectedError));
+			Assert.ExpectError(ExpectedError);
+			ASSERT_THAT(IsFalse(true, ExpectedError));
 		}
+	};
 
-		TEST_METHOD(AssertEqual_WithSameFStrings_Succeeds)
+	TEST_CLASS(NoDiscardAssert_Numbers, "TestFramework.CQTest.Core")
+	{
+		TEST_METHOD(AssertEqual_WithSameInts_Succeeds)
 		{
-			ASSERT_THAT(AreEqual(TEXT("Hello"), TEXT("Hello")));
+			ASSERT_THAT(AreEqual(42, 42));
 		}
-		TEST_METHOD(AssertEqual_WithDifferentFStrings_AddsError)
+		TEST_METHOD(AssertEqual_WithSameIntsAndErrorMessage_DoesNotAddErrorMessage)
 		{
-			Assert.ExpectError(TEXT(""));
-			ASSERT_THAT(AreEqual(TEXT("Hello"), TEXT("World")));
+			ASSERT_THAT(AreEqual(42, 42, "Unexpected"));
 		}
-		TEST_METHOD(AssertEqual_WithDifferentCases_AddsError)
+		TEST_METHOD(AssertEqual_WithDifferentInts_AddsError)
 		{
-			Assert.ExpectError(TEXT(""));
-			ASSERT_THAT(AreEqual(TEXT("hello"), TEXT("HELLO")));
+			Assert.ExpectError(AnyError);
+			ASSERT_THAT(AreEqual(42, 0));
 		}
-		TEST_METHOD(AssertEqual_WithDifferentFStringsAndErrorMessage_AddsSpecificError)
+		TEST_METHOD(AssertEqual_WithDifferentIntsAndErrorMessage_AddsSpecificError)
 		{
-			Assert.ExpectError(FString(ExpectedError));
-			ASSERT_THAT(AreEqual(TEXT("Hello"), TEXT("World"), ExpectedError));
+			Assert.ExpectError(ExpectedError);
+			ASSERT_THAT(AreEqual(42, 0, ExpectedError));
 		}
 		TEST_METHOD(AssertNear_WithSameNumbers_Succeeds)
 		{
@@ -148,15 +155,18 @@ namespace CQTests
 		}
 		TEST_METHOD(AssertNear_WithDifferentNumbers_AddsError)
 		{
-			Assert.ExpectError(TEXT(""));
+			Assert.ExpectError(AnyError);
 			ASSERT_THAT(IsNear(1.0, 2.0, 0.001));
 		}
 		TEST_METHOD(AssertNear_WithDifferentNumbersAndError_AddsSpecificError)
 		{
-			Assert.ExpectError(FString(ExpectedError));
+			Assert.ExpectError(ExpectedError);
 			ASSERT_THAT(IsNear(1.0, 2.0, 0.001, ExpectedError));
 		}
+	};
 
+	TEST_CLASS(NoDiscardAssert_Pointers, "TestFramework.CQTest.Core")
+	{
 		TEST_METHOD(AssertNull_WithNullPtr_Succeeds)
 		{
 			int* ptr = nullptr;
@@ -171,14 +181,14 @@ namespace CQTests
 		{
 			int val = 42;
 			int* ptr = &val;
-			Assert.ExpectError(TEXT(""));
+			Assert.ExpectError(AnyError);
 			ASSERT_THAT(IsNull(ptr));
 		}
 		TEST_METHOD(AssertNull_WithNonNullAndErrorMessage_AddsSpecificErrorMessage)
 		{
 			int val = 42;
 			int* ptr = &val;
-			Assert.ExpectError(FString(ExpectedError));
+			Assert.ExpectError(ExpectedError);
 			ASSERT_THAT(IsNull(ptr, ExpectedError));
 		}
 		TEST_METHOD(AssertNull_WithInvalidSharedPtr_Succeeds)
@@ -189,7 +199,7 @@ namespace CQTests
 		TEST_METHOD(AssertNull_WithValidSharedPtr_AddsError)
 		{
 			TSharedPtr<int> ptr = MakeShared<int>(42);
-			Assert.ExpectError(TEXT(""));
+			Assert.ExpectError(AnyError);
 			ASSERT_THAT(IsNull(ptr));
 		}
 		TEST_METHOD(AssertNull_WithInvalidUniquePtr_Succeeds)
@@ -200,20 +210,20 @@ namespace CQTests
 		TEST_METHOD(AssertNull_WithValidUniquePtr_AddsError)
 		{
 			TUniquePtr<int> ptr = MakeUnique<int>(42);
-			Assert.ExpectError(TEXT(""));
+			Assert.ExpectError(AnyError);
 			ASSERT_THAT(IsNull(ptr));
 		}
 
 		TEST_METHOD(AssertNotNull_WithNullPtr_AddsError)
 		{
 			int* ptr = nullptr;
-			Assert.ExpectError(TEXT(""));
+			Assert.ExpectError(AnyError);
 			ASSERT_THAT(IsNotNull(ptr));
 		}
 		TEST_METHOD(AssertNotNull_WithNullPtrAndError_AddsSpecificError)
 		{
 			int* ptr = nullptr;
-			Assert.ExpectError(FString(ExpectedError));
+			Assert.ExpectError(ExpectedError);
 			ASSERT_THAT(IsNotNull(ptr, ExpectedError));
 		}
 		TEST_METHOD(AssertNotNull_WithNonNull_Succeeds)
@@ -231,7 +241,7 @@ namespace CQTests
 		TEST_METHOD(AssertNotNull_WithInvalidSharedPtr_AddsError)
 		{
 			auto ptr = TSharedPtr<int>{};
-			Assert.ExpectError(TEXT(""));
+			Assert.ExpectError(AnyError);
 			ASSERT_THAT(IsNotNull(ptr));
 		}
 		TEST_METHOD(AssertNotNull_WithValidSharedPtr_Succeeds)
@@ -242,7 +252,7 @@ namespace CQTests
 		TEST_METHOD(AssertNotNull_WithInvalidUniquePtr_AddsError)
 		{
 			TUniquePtr<int> ptr;
-			Assert.ExpectError(TEXT(""));
+			Assert.ExpectError(AnyError);
 			ASSERT_THAT(IsNotNull(ptr));
 		}
 		TEST_METHOD(AssertNotNull_WithValidUniquePtr_Succeeds)
@@ -251,6 +261,186 @@ namespace CQTests
 			ASSERT_THAT(IsNotNull(ptr));
 		}
 	};
+
+	TEST_CLASS(NoDiscardAssert_Strings, "TestFramework.CQTest.Core")
+	{
+		const TCHAR* SomeText{ TEXT("Hello") };
+		const TCHAR* OtherText{ TEXT("World") };
+
+		const TCHAR* UpperText{ TEXT("ONE") };
+		const TCHAR* LowerText{ TEXT("one") };
+
+		const FString SomeString{ SomeText };
+		const FString OtherString{ OtherText };
+
+		TEST_METHOD(AssertAreEqual_Text_Matching_Succeeds)
+		{
+			ASSERT_THAT(AreEqual(SomeText, SomeText));
+		}
+		TEST_METHOD(AssertAreEqual_Text_Mismatch_AddsError)
+		{
+			Assert.ExpectError(AnyError);
+			ASSERT_THAT(AreEqual(SomeText, OtherText));
+		}
+		TEST_METHOD(AssertAreEqual_Text_DifferentCases_AddsError)
+		{
+			Assert.ExpectError(AnyError);
+			ASSERT_THAT(AreEqual(UpperText, LowerText));
+		}
+		TEST_METHOD(AssertAreEqual_Text_MismatchAndError_AddsSpecificError)
+		{
+			Assert.ExpectError(ExpectedError);
+			ASSERT_THAT(AreEqual(SomeText, OtherText, ExpectedError));
+		}
+		TEST_METHOD(AssertAreEqual_FString_Matching_Succeeds)
+		{
+			ASSERT_THAT(AreEqual(SomeString, SomeString));
+		}
+		TEST_METHOD(AssertAreEqual_FString_Mismatch_AddsError)
+		{
+			Assert.ExpectError(AnyError);
+			ASSERT_THAT(AreEqual(SomeString, OtherString));
+		}
+		TEST_METHOD(AssertAreEqual_FString_MismatchAndError_AddsSpecificError)
+		{
+			Assert.ExpectError(ExpectedError);
+			ASSERT_THAT(AreEqual(SomeString, OtherString, ExpectedError));
+		}
+		TEST_METHOD(AssertAreEqual_FString_DifferentCases_AddsError)
+		{
+			Assert.ExpectError(AnyError);
+			ASSERT_THAT(AreEqual(SomeString.ToLower(), SomeString.ToUpper()));
+		}
+
+		TEST_METHOD(AssertAreEqualIgnoreCase_Text_Matching_Succeeds)
+		{
+			ASSERT_THAT(AreEqualIgnoreCase(SomeText, SomeText));
+		}
+		TEST_METHOD(AssertAreEqualIgnoreCase_Text_Mismatch_AddsError)
+		{
+			Assert.ExpectError(AnyError);
+			ASSERT_THAT(AreEqualIgnoreCase(SomeText, OtherText));
+		}
+		TEST_METHOD(AssertAreEqualIgnoreCase_Text_MismatchAndError_AddsSpecificError)
+		{
+			Assert.ExpectError(ExpectedError);
+			ASSERT_THAT(AreEqualIgnoreCase(SomeText, OtherText, ExpectedError));
+		}
+		TEST_METHOD(AssertAreEqualIgnoreCase_Text_DifferentCases_Succeeds)
+		{
+			ASSERT_THAT(AreEqualIgnoreCase(UpperText, LowerText));
+		}
+		TEST_METHOD(AssertAreEqualIgnoreCase_FString_Matching_Succeeds)
+		{
+			ASSERT_THAT(AreEqualIgnoreCase(SomeString, SomeString));
+		}
+		TEST_METHOD(AssertAreEqualIgnoreCase_FString_Mismatch_AddsError)
+		{
+			Assert.ExpectError(AnyError);
+			ASSERT_THAT(AreEqualIgnoreCase(SomeString, OtherString));
+		}
+		TEST_METHOD(AssertAreEqualIgnoreCase_FString_MismatchAndError_AddsSpecificError)
+		{
+			Assert.ExpectError(ExpectedError);
+			ASSERT_THAT(AreEqualIgnoreCase(SomeString, OtherString, ExpectedError));
+		}
+		TEST_METHOD(AssertAreEqualIgnoreCase_FString_DifferentCases_Succeeds)
+		{
+			ASSERT_THAT(AreEqualIgnoreCase(SomeString.ToUpper(), SomeString.ToLower()));
+		}
+
+		TEST_METHOD(AssertAreNotEqual_Text_Matching_AddsError)
+		{
+			Assert.ExpectError(AnyError);
+			ASSERT_THAT(AreNotEqual(SomeText, SomeText));
+		}
+		TEST_METHOD(AssertAreNotEqual_Text_MatchingWithError_AddsSpecificError)
+		{
+			Assert.ExpectError(ExpectedError);
+			ASSERT_THAT(AreNotEqual(SomeText, SomeText, ExpectedError));
+		}
+		TEST_METHOD(AssertAreNotEqual_Text_Mismatch_Succeeds)
+		{
+			ASSERT_THAT(AreNotEqual(SomeText, OtherText));
+		}
+		TEST_METHOD(AssertAreNotEqual_Text_MismatchWithError_Succeeds)
+		{
+			ASSERT_THAT(AreNotEqual(SomeText, OtherText, "Unexpected"));
+		}
+		TEST_METHOD(AssertAreNotEqual_Text_DifferentCases_Succeeds)
+		{
+			ASSERT_THAT(AreNotEqual(UpperText, LowerText));
+		}
+		TEST_METHOD(AssertAreNotEqual_FString_Matching_AddsError)
+		{
+			Assert.ExpectError(AnyError);
+			ASSERT_THAT(AreNotEqual(SomeString, SomeString));
+		}
+		TEST_METHOD(AssertAreNotEqual_FString_MatchingWithMessage_AddsSpecificError)
+		{
+			Assert.ExpectError(ExpectedError);
+			ASSERT_THAT(AreNotEqual(SomeString, SomeString, ExpectedError));
+		}
+		TEST_METHOD(AssertAreNotEqual_FString_Mismatch_Succeeds)
+		{
+			ASSERT_THAT(AreNotEqual(SomeString, OtherString));
+		}
+		TEST_METHOD(AssertAreNotEqual_FString_MismatchAndError_Succeeds)
+		{
+			ASSERT_THAT(AreNotEqual(SomeString, OtherString, ExpectedError));
+		}
+		TEST_METHOD(AssertAreNotEqual_FString_DifferentCases_Succeeds)
+		{
+			ASSERT_THAT(AreNotEqual(SomeString.ToLower(), SomeString.ToUpper()));
+		}
+
+		TEST_METHOD(AssertAreNotEqualIgnoreCase_Text_Matching_AddsError)
+		{
+			Assert.ExpectError(AnyError);
+			ASSERT_THAT(AreNotEqualIgnoreCase(SomeText, SomeText));
+		}
+		TEST_METHOD(AssertAreNotEqualIgnoreCase_Text_MatchingWithError_AddsSpecificError)
+		{
+			Assert.ExpectError(ExpectedError);
+			ASSERT_THAT(AreNotEqualIgnoreCase(SomeText, SomeText));
+		}
+		TEST_METHOD(AssertAreNotEqualIgnoreCase_Text_Mismatch_Succeeds)
+		{
+			ASSERT_THAT(AreNotEqualIgnoreCase(SomeText, OtherText));
+		}
+		TEST_METHOD(AssertAreNotEqualIgnoreCase_Text_MismatchAndError_Succeeds)
+		{
+			ASSERT_THAT(AreNotEqualIgnoreCase(SomeText, OtherText, "Unexpected"));
+		}
+		TEST_METHOD(AssertAreNotEqualIgnoreCase_Text_DifferentCases_AddsError)
+		{
+			Assert.ExpectError(AnyError);
+			ASSERT_THAT(AreNotEqualIgnoreCase(UpperText, LowerText));
+		}
+		TEST_METHOD(AssertAreNotEqualIgnoreCase_FString_Matching_AddsError)
+		{
+			Assert.ExpectError(AnyError);
+			ASSERT_THAT(AreNotEqualIgnoreCase(SomeString, SomeString));
+		}
+		TEST_METHOD(AssertAreNotEqualIgnoreCase_FString_MatchingWithError_AddsSpecificError)
+		{
+			Assert.ExpectError(ExpectedError);
+			ASSERT_THAT(AreNotEqualIgnoreCase(SomeString, SomeString, ExpectedError));
+		}
+		TEST_METHOD(AssertAreNotEqualIgnoreCase_FString_Mismatch_Succeeds)
+		{
+			ASSERT_THAT(AreNotEqualIgnoreCase(SomeString, OtherString));
+		}
+		TEST_METHOD(AssertAreNotEqualIgnoreCase_FString_MismatchAndError_Succeeds)
+		{
+			ASSERT_THAT(AreNotEqualIgnoreCase(SomeString, OtherString, "Unexpected"));
+		}
+		TEST_METHOD(AssertAreNotEqualIgnoreCase_FString_DifferentCases_AddsError)
+		{
+			Assert.ExpectError(AnyError);
+			ASSERT_THAT(AreNotEqualIgnoreCase(SomeString.ToUpper(), SomeString.ToLower()));
+		}
+	};
 	
 } // namespace CQTests
 
diff --git a/Engine/Plugins/Tests/CQTest/Source/CQTestTests/Private/Assert/CQTestConvertTests.cpp b/Engine/Plugins/Tests/CQTest/Source/CQTestTests/Private/Assert/CQTestConvertTests.cpp
--- a/Engine/Plugins/Tests/CQTest/Source/CQTestTests/Private/Assert/CQTestConvertTests.cpp
+++ b/Engine/Plugins/Tests/CQTest/Source/CQTestTests/Private/Assert/CQTestConvertTests.cpp
@@ -22,26 +22,38 @@ struct TypeWithToString
 	{
 		return ToStringStr;
 	}
+
+	bool operator==(const TypeWithToString& A) const
+	{
+		return true;
+	}
+	bool operator!=(const TypeWithToString& A) const
+	{
+		return false;
+	}
 };
 
 struct TypeWithOStream
 {
 };
 
-CQTestConvert::PrintStream& operator<<(CQTestConvert::PrintStream& stream, const TypeWithOStream& obj) {
+CQTestConvert::PrintStream& operator<<(CQTestConvert::PrintStream& stream, const TypeWithOStream& obj)
+{
 	stream << OStreamStr;
 	return stream;
 }
 
 struct TypeWithToStringAndOverride
 {
-	FString ToString() const {
+	FString ToString() const
+	{
 		return ToStringStr;
 	}
 };
 
 template<>
-inline FString CQTestConvert::ToString(const TypeWithToStringAndOverride&) {
+inline FString CQTestConvert::ToString(const TypeWithToStringAndOverride&)
+{
 	return OverrideStr;
 }
 
@@ -63,10 +75,22 @@ enum UnscopedEnumNoOperator
 	Two = 2
 };
 
-enum struct ScopedEnumNoOperator
+enum UnscopedEnumWithUCharBase : uint8
 {
-	One = 1,
-	Two = 2
+	First = 1,
+	Second = 2
+};
+
+enum UnscopedEnumWithCharBase : int8
+{
+	Alpha = 1,
+	Beta = 2
+};
+
+enum UnscopedEnumWithIntBase : int32
+{
+	Aleph = 1,
+	Bet = 2
 };
 
 enum UnscopedEnumWithOperator
@@ -75,6 +99,36 @@ enum UnscopedEnumWithOperator
 	Four = 4
 };
 
+enum struct ScopedEnumNoOperator
+{
+	One = 1,
+	Two = 2
+};
+
+enum struct ScopedEnumWithOperator
+{
+	One = 1,
+	Two = 2
+};
+
+enum class ScopedEnumWithUCharBase : uint8
+{
+	One = 1,
+	Two = 2
+};
+
+enum class ScopedEnumWithCharBase : int8
+{
+	One = 1,
+	Two = 2
+};
+
+enum class ScopedEnumWithIntBase : int32
+{
+	One = 1,
+	Two = 2
+};
+
 CQTestConvert::PrintStream& operator<<(CQTestConvert::PrintStream& stream, const UnscopedEnumWithOperator& e)
 {
 	switch (e)
@@ -92,12 +146,6 @@ CQTestConvert::PrintStream& operator<<(CQTestConvert::PrintStream& stream, const
 	return stream;
 }
 
-enum struct ScopedEnumWithOperator
-{
-	One = 1,
-	Two = 2
-};
-
 CQTestConvert::PrintStream& operator<<(CQTestConvert::PrintStream& stream, const ScopedEnumWithOperator& e)
 {
 	switch (e)
@@ -147,6 +195,14 @@ TEST_CLASS(TestConvertToString, "TestFramework.CQTest.Core")
 		ASSERT_THAT(AreEqual(ToStringStr, CQTestConvert::ToString(obj)));
 	}
 
+	TEST_METHOD(TypeWithToString_ConvertToString_TestNotEqualToString)
+	{
+		TypeWithToString obj1;
+		TypeWithToString obj2;
+		Assert.ExpectError(FString::Printf(TEXT("Expected %s to not equal %s"), *obj1.ToString(), *obj2.ToString()));
+		ASSERT_THAT(AreNotEqual(obj1, obj2));
+	}
+
 	TEST_METHOD(UnscopedEnumNoOperator_ConvertToString_ReturnsNumber)
 	{
 		auto e = UnscopedEnumNoOperator::One;
@@ -171,6 +227,42 @@ TEST_CLASS(TestConvertToString, "TestFramework.CQTest.Core")
 		ASSERT_THAT(AreEqual(FString(TEXT("Two")), CQTestConvert::ToString(e)));
 	}
 
+	TEST_METHOD(ScopedEnumWithUCharBase_ConverToString_ReturnsNumberAsString)
+	{
+		auto e = ScopedEnumWithUCharBase::One;
+		ASSERT_THAT(AreEqual(FString(TEXT("1")), CQTestConvert::ToString(e)));
+	}
+
+	TEST_METHOD(UnscopedEnumWithUCharBase_ConvertToString_ReturnsNumberAsString)
+	{
+		auto e = UnscopedEnumWithUCharBase::First;
+		ASSERT_THAT(AreEqual(FString(TEXT("1")), CQTestConvert::ToString(e)));
+	}
+
+	TEST_METHOD(ScopedEnumWithCharBase_ConverToString_ReturnsNumberAsString)
+	{
+		auto e = ScopedEnumWithCharBase::One;
+		ASSERT_THAT(AreEqual(FString(TEXT("1")), CQTestConvert::ToString(e)));
+	}
+
+	TEST_METHOD(UnscopedEnumWithCharBase_ConvertToString_ReturnsNumberAsString)
+	{
+		auto e = UnscopedEnumWithCharBase::Alpha;
+		ASSERT_THAT(AreEqual(FString(TEXT("1")), CQTestConvert::ToString(e)));
+	}
+
+	TEST_METHOD(ScopedEnumWithIntBase_ConverToString_ReturnsNumberAsString)
+	{
+		auto e = ScopedEnumWithIntBase::One;
+		ASSERT_THAT(AreEqual(FString(TEXT("1")), CQTestConvert::ToString(e)));
+	}
+
+	TEST_METHOD(UnscopedEnumWithIntBase_ConvertToString_ReturnsNumberAsString)
+	{
+		auto e = UnscopedEnumWithIntBase::Aleph;
+		ASSERT_THAT(AreEqual(FString(TEXT("1")), CQTestConvert::ToString(e)));
+	}
+
 	TEST_METHOD(ArrayOfValues_ConvertToString_ReturnsCommaDelimitedList)
 	{
 		TArray<int32> values = { 1, 2, 3 };
@@ -199,4 +291,11 @@ TEST_CLASS(TestConvertToString, "TestFramework.CQTest.Core")
 		ASSERT_THAT(IsTrue(asStr.Contains(TEXT("2:Two"))));
 		ASSERT_THAT(IsTrue(asStr.Contains(TEXT("3:Three"))));
 	}
+
+	TEST_METHOD(Float_ConvertToString_ReturnsFloatAsString)
+	{
+		double value = 3.14;
+		auto asStr = CQTestConvert::ToString(value);
+		ASSERT_THAT(IsTrue(asStr.Contains(TEXT("3.14"))));
+	}
 };
