[CQTest] Fixed FWaitUntil grabbing start time at construction instead of on first update
Fixed ActorTestSpawner's tear down of objects

#rnx

diff --git a/Engine/Plugins/Tests/CQTest/Source/CQTest/Private/Commands/TestCommands.cpp b/Engine/Plugins/Tests/CQTest/Source/CQTest/Private/Commands/TestCommands.cpp
--- a/Engine/Plugins/Tests/CQTest/Source/CQTest/Private/Commands/TestCommands.cpp
+++ b/Engine/Plugins/Tests/CQTest/Source/CQTest/Private/Commands/TestCommands.cpp
@@ -7,33 +7,38 @@ DEFINE_LOG_CATEGORY_STATIC(LogCqTest, Log, All)
 
 bool CQTEST_API FWaitUntil::Update()
 {
-	if (Description != nullptr && !bHasLoggedStart)
+	if (!bHasTimerStarted)
 	{
-		UE_LOG(LogCqTest, Log, TEXT("Starting %s"), Description);
-		bHasLoggedStart = true;
+		StartTime = FDateTime::UtcNow();
+		bHasTimerStarted = true;
+		if (Description != nullptr)
+		{
+			UE_LOG(LogCqTest, Log, TEXT("Starting %s"), Description);
+		}
 	}
 	if (TestRunner.HasAnyErrors())
 	{
 		return true;
 	}
 
+	auto Elapsed = FDateTime::UtcNow() - StartTime;
 	if (Query())
 	{
 		if (Description)
 		{
-			UE_LOG(LogCqTest, Log, TEXT("Finished %s"), Description);
+			UE_LOG(LogCqTest, Log, TEXT("Finished %s after %d milliseconds"), Description, Elapsed.GetTotalMilliseconds());
 		}
 		return true;
 	}
-	else if (FDateTime::UtcNow() >= StartTime + Timeout)
+	else if (Elapsed >= Timeout)
 	{
 		if (Description)
 		{
-			TestRunner.AddError(*FString::Printf(TEXT("Timed out waiting for %s"), Description));
+			TestRunner.AddError(*FString::Printf(TEXT("Timed out waiting for %s after %d milliseconds"), Description, Elapsed.GetTotalMilliseconds()));
 		}
 		else
 		{
-			TestRunner.AddError(TEXT("Latent command timed out."));
+			TestRunner.AddError(*FString::Printf(TEXT("Latent command timed out after %d milliseconds"), Elapsed.GetTotalMilliseconds()));
 		}
 		return true;
 	}
diff --git a/Engine/Plugins/Tests/CQTest/Source/CQTest/Private/Components/ActorTestSpawner.cpp b/Engine/Plugins/Tests/CQTest/Source/CQTest/Private/Components/ActorTestSpawner.cpp
--- a/Engine/Plugins/Tests/CQTest/Source/CQTest/Private/Components/ActorTestSpawner.cpp
+++ b/Engine/Plugins/Tests/CQTest/Source/CQTest/Private/Components/ActorTestSpawner.cpp
@@ -50,7 +50,7 @@ namespace
 		{
 			if (UObject* Obj = InSpawnedObjects[Index].Get())
 			{
-				Obj->BeginDestroy();
+				Obj->ConditionalBeginDestroy();
 			}
 		}
 
@@ -66,10 +66,11 @@ FActorTestSpawner::~FActorTestSpawner() {
 		GameInstance = nullptr;
 	}
 
+	DestroySpawnedActors(SpawnedActors);
+	DestroySpawnedObjects(SpawnedObjects);
+
 	if (GameWorld)
 	{
-		DestroySpawnedActors(SpawnedActors);
-		DestroySpawnedObjects(SpawnedObjects);
 		DestroyWorld(GameWorld);
 		GameWorld->RemoveFromRoot();
 		GameWorld = nullptr;
diff --git a/Engine/Plugins/Tests/CQTest/Source/CQTest/Private/Components/PIENetworkComponent.cpp b/Engine/Plugins/Tests/CQTest/Source/CQTest/Private/Components/PIENetworkComponent.cpp
--- a/Engine/Plugins/Tests/CQTest/Source/CQTest/Private/Components/PIENetworkComponent.cpp
+++ b/Engine/Plugins/Tests/CQTest/Source/CQTest/Private/Components/PIENetworkComponent.cpp
@@ -107,9 +107,16 @@ void FBasePIENetworkComponent::StartPie()
 
 	FRequestPlaySessionParams SessionParams;
 	SessionParams.WorldType = EPlaySessionWorldType::PlayInEditor;
-	SessionParams.GameModeOverride = AGameModeBase::StaticClass();
 	SessionParams.DestinationSlateViewport = LevelEditorModule.GetFirstActiveViewport();
 	SessionParams.EditorPlaySettings = PlaySettings;
+	if (GameMode != nullptr)
+	{
+		SessionParams.GameModeOverride = GameMode;
+	}
+	else
+	{
+		SessionParams.GameModeOverride = AGameModeBase::StaticClass();
+	}
 
 	GUnrealEd->RequestPlaySession(SessionParams);
 	GUnrealEd->StartQueuedPlaySessionRequest();
diff --git a/Engine/Plugins/Tests/CQTest/Source/CQTest/Public/Commands/TestCommands.h b/Engine/Plugins/Tests/CQTest/Source/CQTest/Public/Commands/TestCommands.h
--- a/Engine/Plugins/Tests/CQTest/Source/CQTest/Public/Commands/TestCommands.h
+++ b/Engine/Plugins/Tests/CQTest/Source/CQTest/Public/Commands/TestCommands.h
@@ -11,7 +11,6 @@ public:
 		: TestRunner(InTestRunner)
 		, Query(MoveTemp(Query)) 
 		, Timeout(Timeout)
-		, StartTime(FDateTime::UtcNow())
 		, Description(InDescription)
 	{}
 
@@ -22,7 +21,7 @@ public:
 	FTimespan Timeout;
 	FDateTime StartTime;
 	const TCHAR* Description;
-	bool bHasLoggedStart = false;
+	bool bHasTimerStarted = false;
 };
 
 enum class ECQTestFailureBehavior
@@ -67,7 +66,8 @@ public:
 	void AppendAll(TArray < TSharedPtr<IAutomationLatentCommand>> ToAdd);
 	void Prepend(TSharedPtr<IAutomationLatentCommand> ToAdd);
 
-	bool IsEmpty() const {
+	bool IsEmpty() const
+	{
 		return Commands.IsEmpty();
 	}
 
diff --git a/Engine/Plugins/Tests/CQTest/Source/CQTest/Public/Components/PIENetworkComponent.h b/Engine/Plugins/Tests/CQTest/Source/CQTest/Public/Components/PIENetworkComponent.h
--- a/Engine/Plugins/Tests/CQTest/Source/CQTest/Public/Components/PIENetworkComponent.h
+++ b/Engine/Plugins/Tests/CQTest/Source/CQTest/Public/Components/PIENetworkComponent.h
@@ -50,10 +50,11 @@ protected:
 	FAutomationTestBase* TestRunner = nullptr;
 	FTestCommandBuilder* CommandBuilder = nullptr;
 	FPacketSimulationSettings* PacketSimulationSettings = nullptr;
+	TSubclassOf<AGameModeBase> GameMode = TSubclassOf<AGameModeBase>(nullptr);
 	FPIENetworkTestStateRestorer StateRestorer;
 };
 
-template <typename FBasePIENetworkComponentState>
+template <typename NetworkDataType>
 class FNetworkComponentBuilder;
 
 template <typename NetworkDataType = FBasePIENetworkComponentState>
diff --git a/Engine/Plugins/Tests/CQTest/Source/CQTest/Public/Impl/PIENetworkComponent.inl b/Engine/Plugins/Tests/CQTest/Source/CQTest/Public/Impl/PIENetworkComponent.inl
--- a/Engine/Plugins/Tests/CQTest/Source/CQTest/Public/Impl/PIENetworkComponent.inl
+++ b/Engine/Plugins/Tests/CQTest/Source/CQTest/Public/Impl/PIENetworkComponent.inl
@@ -220,5 +220,6 @@ inline void FNetworkComponentBuilder<NetworkDataType>::Build(FPIENetworkComponen
 	}
 
 	OutNetwork.PacketSimulationSettings = PacketSimulationSettings;
+	OutNetwork.GameMode = GameMode;
 	OutNetwork.StateRestorer = FPIENetworkTestStateRestorer{GameInstanceClass, GameMode};
 }
diff --git a/Engine/Plugins/Tests/CQTest/Source/CQTestTests/Private/Components/PIENetworkComponentTests.cpp b/Engine/Plugins/Tests/CQTest/Source/CQTestTests/Private/Components/PIENetworkComponentTests.cpp
--- a/Engine/Plugins/Tests/CQTest/Source/CQTestTests/Private/Components/PIENetworkComponentTests.cpp
+++ b/Engine/Plugins/Tests/CQTest/Source/CQTestTests/Private/Components/PIENetworkComponentTests.cpp
@@ -99,7 +99,7 @@ NETWORK_TEST_CLASS(ReplicationTest, "TestFramework.CQTest.Network")
 			.Build(Network);
 	}
 
-	TEST_METHOD(SpawnActor_WithReplicatedActor_ProvidesActorToClients)
+	TEST_METHOD(SpawnAndReplicateActor_WithReplicatedActor_ProvidesActorToClients)
 	{
 		Network.SpawnAndReplicate<ATestReplicatedActor, &DerivedState::ReplicatedActor>()
 			.ThenServer([&](DerivedState& ServerState) {
diff --git a/Engine/Plugins/Tests/CQTest/Source/CQTestTests/Private/Components/SpawnHelperTests.cpp b/Engine/Plugins/Tests/CQTest/Source/CQTestTests/Private/Components/SpawnHelperTests.cpp
--- a/Engine/Plugins/Tests/CQTest/Source/CQTestTests/Private/Components/SpawnHelperTests.cpp
+++ b/Engine/Plugins/Tests/CQTest/Source/CQTestTests/Private/Components/SpawnHelperTests.cpp
@@ -27,4 +27,9 @@ TEST_CLASS(SpawnHelperTests, "TestFramework.CQTest.Actor")
 		ASSERT_THAT(IsNotNull(TestComponent));
 		ASSERT_THAT(AreEqual(TestComponent->GetOwner(), &Owner ));
 	}
+
+	TEST_METHOD(SpawnObject_WithAnbject_TearsDownCleanly)
+	{
+		auto& SomeComponent = Spawner.SpawnObject<USceneComponent>();
+	}
 };
\ No newline at end of file
