[ CQTest] Remove a duplicated test

#rnx

diff --git a/Engine/Plugins/Tests/CQTest/Source/CQTestTests/Private/Assert/CQTestConvertTests.cpp b/Engine/Plugins/Tests/CQTest/Source/CQTestTests/Private/Assert/CQTestConvertTests.cpp
--- a/Engine/Plugins/Tests/CQTest/Source/CQTestTests/Private/Assert/CQTestConvertTests.cpp
+++ b/Engine/Plugins/Tests/CQTest/Source/CQTestTests/Private/Assert/CQTestConvertTests.cpp
@@ -22,15 +22,6 @@ struct TypeWithToString
 	{
 		return ToStringStr;
 	}
-
-	bool operator==(const TypeWithToString& A) const
-	{
-		return true;
-	}
-	bool operator!=(const TypeWithToString& A) const
-	{
-		return false;
-	}
 };
 
 struct TypeWithOStream
@@ -195,14 +186,6 @@ TEST_CLASS(TestConvertToString, "TestFramework.CQTest.Core")
 		ASSERT_THAT(AreEqual(ToStringStr, CQTestConvert::ToString(obj)));
 	}
 
-	TEST_METHOD(TypeWithToString_ConvertToString_TestNotEqualToString)
-	{
-		TypeWithToString obj1;
-		TypeWithToString obj2;
-		Assert.ExpectError(FString::Printf(TEXT("Expected %s to not equal %s"), *obj1.ToString(), *obj2.ToString()));
-		ASSERT_THAT(AreNotEqual(obj1, obj2));
-	}
-
 	TEST_METHOD(UnscopedEnumNoOperator_ConvertToString_ReturnsNumber)
 	{
 		auto e = UnscopedEnumNoOperator::One;
