[CQTest] avoid instanting test when no automation worker are built

#rnx

diff --git a/Engine/Plugins/Tests/CQTest/Source/CQTest/Public/CQTest.h b/Engine/Plugins/Tests/CQTest/Source/CQTest/Public/CQTest.h
--- a/Engine/Plugins/Tests/CQTest/Source/CQTest/Public/CQTest.h
+++ b/Engine/Plugins/Tests/CQTest/Source/CQTest/Public/CQTest.h
@@ -154,15 +154,26 @@ struct TTest : TBaseTest<AsserterType>
 	inline static TTestRunner<AsserterType>* TestRunner{ nullptr };
 };
 
-#define _TEST_CLASS_IMPL(_ClassName, _TestDir, _BaseClass, _AsserterType, _TestFlags)                                                 \
-	struct _ClassName;                                                                                                    \
-	struct F##_ClassName##_Runner : public TTestRunner<_AsserterType>                                                     \
-	{                                                                                                                     \
-		F##_ClassName##_Runner()                                                                                          \
-			: TTestRunner(#_ClassName, __LINE__, __FILE__, _TestDir, _TestFlags, TTest<_ClassName, _AsserterType>::CreateTestClass) {} \
-	};                                                                                                                    \
-	F##_ClassName##_Runner _ClassName##_RunnerInstance;                                                                   \
-	struct _ClassName : public _BaseClass<_ClassName, _AsserterType>
+#if WITH_AUTOMATION_WORKER
+	#define _TEST_CLASS_IMPL(_ClassName, _TestDir, _BaseClass, _AsserterType, _TestFlags)                                                  \
+		struct _ClassName;                                                                                                                 \
+		struct F##_ClassName##_Runner : public TTestRunner<_AsserterType>                                                                  \
+		{                                                                                                                                  \
+			F##_ClassName##_Runner()                                                                                                       \
+				: TTestRunner(#_ClassName, __LINE__, __FILE__, _TestDir, _TestFlags, TTest<_ClassName, _AsserterType>::CreateTestClass) {} \
+		};                                                                                                                                 \
+		F##_ClassName##_Runner _ClassName##_RunnerInstance;                                                                                \
+		struct _ClassName : public _BaseClass<_ClassName, _AsserterType>
+
+	#define TEST_METHOD(_MethodName)                                                                       \
+		FFunctionRegistrar reg##_MethodName{ FString(#_MethodName), &DerivedType::_MethodName, __LINE__ }; \
+	void _MethodName()
+#else
+	#define _TEST_CLASS_IMPL(_ClassName, _TestDir, _BaseClass, _AsserterType, _TestFlags) \
+			struct _ClassName : public _BaseClass<_ClassName, _AsserterType>
+
+	#define TEST_METHOD(_MethodName) void _MethodName()
+#endif
 
 #define TEST_CLASS_WITH_ASSERTS(_ClassName, _TestDir, _AsserterType) _TEST_CLASS_IMPL(_ClassName, _TestDir, TTest, _AsserterType, DefaultFlags)
 
@@ -171,10 +182,6 @@ struct TTest : TBaseTest<AsserterType>
 #define TEST_CLASS_WITH_FLAGS(_ClassName, _TestDir, _Flags) _TEST_CLASS_IMPL(_ClassName, _TestDir, TTest, FNoDiscardAsserter, _Flags)
 #define TEST_CLASS_WITH_BASE_AND_FLAGS(_ClassName, _TestDir, _BaseClass, _Flags) _TEST_CLASS_IMPL(_ClassName, _TestDir, _BaseClass, FNoDiscardAsserter, _Flags)
 
-#define TEST_METHOD(_MethodName)                                                            \
-	FFunctionRegistrar reg##_MethodName{ FString(#_MethodName), &DerivedType::_MethodName, __LINE__ }; \
-	void _MethodName()
-
 #define TEST(_TestName, _TestDir)        \
 	TEST_CLASS(_TestName, _TestDir)      \
 	{                                    \
