[CQTest] Use Transient package to hold transient world from Actor spawner helper

#rnx

diff --git a/Engine/Plugins/Tests/CQTest/Source/CQTest/Private/Components/ActorTestSpawner.cpp b/Engine/Plugins/Tests/CQTest/Source/CQTest/Private/Components/ActorTestSpawner.cpp
--- a/Engine/Plugins/Tests/CQTest/Source/CQTest/Private/Components/ActorTestSpawner.cpp
+++ b/Engine/Plugins/Tests/CQTest/Source/CQTest/Private/Components/ActorTestSpawner.cpp
@@ -79,8 +79,9 @@ FActorTestSpawner::~FActorTestSpawner() {
 
 UWorld* FActorTestSpawner::CreateWorld()
 {
+	FName WorldName = MakeUniqueObjectName(nullptr, UWorld::StaticClass());
 	FWorldContext& WorldContext = GEngine->CreateNewWorldContext(EWorldType::Game);
-	UWorld* Result = UWorld::CreateWorld(EWorldType::Game, false, NAME_None);
+	UWorld* Result = UWorld::CreateWorld(EWorldType::Game, false, WorldName, GetTransientPackage());
 	check(Result != nullptr);
 	Result->AddToRoot();
 	WorldContext.SetCurrentWorld(Result);
