Fixing UWorld name conflicts by making sure that the name is globally unique when creating the UWorld object for testing. Profiling the change under the Development Editor has shown a negligible hit in performance (less than a second) when running the CQTests with this change.

#rnx

diff --git a/Engine/Plugins/Tests/CQTest/Source/CQTest/Private/Components/ActorTestSpawner.cpp b/Engine/Plugins/Tests/CQTest/Source/CQTest/Private/Components/ActorTestSpawner.cpp
--- a/Engine/Plugins/Tests/CQTest/Source/CQTest/Private/Components/ActorTestSpawner.cpp
+++ b/Engine/Plugins/Tests/CQTest/Source/CQTest/Private/Components/ActorTestSpawner.cpp
@@ -79,7 +79,7 @@ FActorTestSpawner::~FActorTestSpawner() {
 
 UWorld* FActorTestSpawner::CreateWorld()
 {
-	FName WorldName = MakeUniqueObjectName(nullptr, UWorld::StaticClass());
+	FName WorldName = MakeUniqueObjectName(nullptr, UWorld::StaticClass(), NAME_None, EUniqueObjectNameOptions::GloballyUnique);
 	FWorldContext& WorldContext = GEngine->CreateNewWorldContext(EWorldType::Game);
 	UWorld* Result = UWorld::CreateWorld(EWorldType::Game, false, WorldName, GetTransientPackage());
 	check(Result != nullptr);
