[CQTest] Exposed the methods for AddInfo, AddWarning, AddError, AddErrorIfFalse, and AddEvent to TBaseTest. Marked AddErrorS and AddWarningS for deprecation.

#rnx

diff --git a/Engine/Plugins/Tests/CQTest/Source/CQTest/Public/CQTest.h b/Engine/Plugins/Tests/CQTest/Source/CQTest/Public/CQTest.h
--- a/Engine/Plugins/Tests/CQTest/Source/CQTest/Public/CQTest.h
+++ b/Engine/Plugins/Tests/CQTest/Source/CQTest/Public/CQTest.h
@@ -70,6 +70,36 @@ struct TBaseTest
 	void AddCommand(IAutomationLatentCommand* Cmd);
 	void AddCommand(TSharedPtr<IAutomationLatentCommand> Cmd);
 
+	/**
+	 * Adds an error message to this test
+	 *
+	 * @param	InError	Error message to add to this test
+	 */
+	void AddError(const FString& InError) const;
+
+	/**
+	 * Adds an error message to this test if the condition is false
+	 *
+	 * @param   bCondition	The condition to validate.
+	 * @param   InError		Error message to add to this test
+	 * @return	False if there was an error
+	 */
+	bool AddErrorIfFalse(bool bCondition, const FString& InError) const;
+
+	/**
+	 * Adds a warning to this test
+	 *
+	 * @param	InWarning	Warning message to add to this test
+	 */
+	void AddWarning(const FString& InWarning) const;
+
+	/**
+	 * Adds a log item to this test
+	 *
+	 * @param	InLogItem	Log item to add to this test
+	 */
+	void AddInfo(const FString& InLogItem) const;
+
 	virtual void RunTest(const FString& MethodName) = 0;
 
 	bool bInitializing{ true };
diff --git a/Engine/Plugins/Tests/CQTest/Source/CQTest/Public/Impl/CQTest.inl b/Engine/Plugins/Tests/CQTest/Source/CQTest/Public/Impl/CQTest.inl
--- a/Engine/Plugins/Tests/CQTest/Source/CQTest/Public/Impl/CQTest.inl
+++ b/Engine/Plugins/Tests/CQTest/Source/CQTest/Public/Impl/CQTest.inl
@@ -306,3 +306,27 @@ inline void TBaseTest<AsserterType>::AddCommand(TSharedPtr<IAutomationLatentComm
 {
 	TestCommandBuilder.CommandQueue.Add(Cmd);
 }
+
+template <typename AsserterType>
+inline void TBaseTest<AsserterType>::AddError(const FString& InError) const
+{
+	TestRunner.AddError(InError, 0);
+}
+
+template <typename AsserterType>
+inline bool TBaseTest<AsserterType>::AddErrorIfFalse(bool bCondition, const FString& InError) const
+{
+	return TestRunner.AddErrorIfFalse(bCondition, InError, 0);
+}
+
+template <typename AsserterType>
+inline void TBaseTest<AsserterType>::AddWarning(const FString& InWarning) const
+{
+	TestRunner.AddWarning(InWarning, 0);
+}
+
+template <typename AsserterType>
+inline void TBaseTest<AsserterType>::AddInfo(const FString& InLogItem) const
+{
+	TestRunner.AddInfo(InLogItem, 0, false);
+}
\ No newline at end of file
diff --git a/Engine/Source/Runtime/Core/Public/Misc/AutomationTest.h b/Engine/Source/Runtime/Core/Public/Misc/AutomationTest.h
--- a/Engine/Source/Runtime/Core/Public/Misc/AutomationTest.h
+++ b/Engine/Source/Runtime/Core/Public/Misc/AutomationTest.h
@@ -1412,6 +1412,7 @@ public:
 	 * @param	InFilename	The filename the error originated in
 	 * @param	InLineNumber	The line number in the file this error originated in
 	 */
+	UE_DEPRECATED(5.4, "Please use AddError instead.")
 	CORE_API virtual void AddErrorS(const FString& InError, const FString& InFilename, int32 InLineNumber);
 
 	/**
@@ -1421,6 +1422,7 @@ public:
 	 * @param	InFilename	The filename the error originated in
 	 * @param	InLineNumber	The line number in the file this error originated in
 	 */
+	UE_DEPRECATED(5.4, "Please use AddWarning instead.")
 	CORE_API virtual void AddWarningS(const FString& InWarning, const FString& InFilename, int32 InLineNumber);
 
 	/**
