[CQTest] Added tests for newly exposed AddInfo, AddWarning, and AddErrorIfFalse methods.

#rnx

diff --git a/Engine/Plugins/Tests/CQTest/Source/CQTestTests/Private/CQTestTests.cpp b/Engine/Plugins/Tests/CQTest/Source/CQTestTests/Private/CQTestTests.cpp
--- a/Engine/Plugins/Tests/CQTest/Source/CQTestTests/Private/CQTestTests.cpp
+++ b/Engine/Plugins/Tests/CQTest/Source/CQTestTests/Private/CQTestTests.cpp
@@ -484,4 +484,53 @@ namespace CQTestTests
 			UE_LOG(TestSuppressLog, Error, TEXT("%s"), *ExpectedError);
 		}
 	};
+
+	TEST_CLASS(AutomationEvents, "TestFramework.CQTest.Core")
+	{
+		bool bIsExpectingError{ false };
+		bool bIsExpectingWarning{ false };
+		FString ExpectedMessage = TEXT("ExpectedMessage");
+
+		AFTER_EACH()
+		{
+			if (bIsExpectingError)
+			{
+				ClearExpectedErrors(*this->TestRunner, { TEXT("will be marked as failing due to errors"), ExpectedMessage });
+			}
+			if (bIsExpectingWarning)
+			{
+				ClearExpectedWarning(*this->TestRunner, ExpectedMessage);
+			}
+		}
+
+		TEST_METHOD(InfoEvent)
+		{
+			AddInfo(ExpectedMessage);
+			ASSERT_THAT(IsTrue(DoesEventExist(*this->TestRunner, FAutomationEvent(EAutomationEventType::Info, ExpectedMessage))));
+		}
+
+		TEST_METHOD(WarningEvent)
+		{
+			bIsExpectingWarning = true;
+			AddWarning(ExpectedMessage);
+			ASSERT_THAT(IsTrue(DoesEventExist(*this->TestRunner, FAutomationEvent(EAutomationEventType::Warning, ExpectedMessage))));
+		}
+
+		TEST_METHOD(ConditionalEvent)
+		{
+			FString InvalidErrorMessage = TEXT("This should not be an error");
+			AddErrorIfFalse(true, InvalidErrorMessage);
+			ASSERT_THAT(IsFalse(DoesEventExist(*this->TestRunner, FAutomationEvent(EAutomationEventType::Error, InvalidErrorMessage))));
+			AddErrorIfFalse(false, ExpectedMessage);
+			bIsExpectingError = true;
+			ASSERT_THAT(IsTrue(DoesEventExist(*this->TestRunner, FAutomationEvent(EAutomationEventType::Error, ExpectedMessage))));
+		}
+
+		TEST_METHOD(ErrorEvent)
+		{
+			bIsExpectingError = true;
+			AddError(ExpectedMessage);
+			ASSERT_THAT(IsTrue(DoesEventExist(*this->TestRunner, FAutomationEvent(EAutomationEventType::Error, ExpectedMessage))));
+		}
+	};
 } // namespace CQTestTests
diff --git a/Engine/Plugins/Tests/CQTest/Source/CQTestTests/Private/CQTestUnitTestHelper.cpp b/Engine/Plugins/Tests/CQTest/Source/CQTestTests/Private/CQTestUnitTestHelper.cpp
--- a/Engine/Plugins/Tests/CQTest/Source/CQTestTests/Private/CQTestUnitTestHelper.cpp
+++ b/Engine/Plugins/Tests/CQTest/Source/CQTestTests/Private/CQTestUnitTestHelper.cpp
@@ -2,6 +2,25 @@
 
 #include "CQTestUnitTestHelper.h"
 
+bool DoesEventExist(FAutomationTestBase& TestRunner, const FAutomationEvent& Event)
+{
+	FAutomationTestExecutionInfo TestInfo;
+	TestRunner.GetExecutionInfo(TestInfo);
+	const TArray<FAutomationExecutionEntry>& TestEntries = TestInfo.GetEntries();
+	if (TestEntries.Num() > 0)
+	{
+		for (const FAutomationExecutionEntry& Entry : TestEntries)
+		{
+			if (Entry.Event.Type == Event.Type && Entry.Event.Message.Equals(Event.Message))
+			{
+				return true;
+			}
+		}
+	}
+
+	return false;
+}
+
 void ClearExpectedErrors(FAutomationTestBase& TestRunner, const TArray<FString>& ExpectedErrors)
 {
 	FAutomationTestExecutionInfo testInfo;
diff --git a/Engine/Plugins/Tests/CQTest/Source/CQTestTests/Private/CQTestUnitTestHelper.h b/Engine/Plugins/Tests/CQTest/Source/CQTestTests/Private/CQTestUnitTestHelper.h
--- a/Engine/Plugins/Tests/CQTest/Source/CQTestTests/Private/CQTestUnitTestHelper.h
+++ b/Engine/Plugins/Tests/CQTest/Source/CQTestTests/Private/CQTestUnitTestHelper.h
@@ -3,6 +3,8 @@
 #pragma once
 #include "Misc/AutomationTest.h"
 
+bool DoesEventExist(FAutomationTestBase& TestRunner, const FAutomationEvent& Event);
+
 // When testing the framework, occasionally we need the error to live on the test framework
 // for short-circuiting logic.  We use this method to manually clear the error
 // (Normally from AFTER_EACH) to cause the test to pass
