Updating CQTest to compile using Cpp20 from Cpp17

#rnx

diff --git a/Engine/Plugins/Tests/CQTest/Source/CQTest/CQTest.Build.cs b/Engine/Plugins/Tests/CQTest/Source/CQTest/CQTest.Build.cs
--- a/Engine/Plugins/Tests/CQTest/Source/CQTest/CQTest.Build.cs
+++ b/Engine/Plugins/Tests/CQTest/Source/CQTest/CQTest.Build.cs
@@ -7,15 +7,7 @@ public class CQTest : ModuleRules
 	public CQTest(ReadOnlyTargetRules Target)
 		: base(Target)
 	{
-		// Does not compile with C++20:
-		// error C2088: '<<': illegal for class
-		// error C2280: 'std::basic_ostream<char,std::char_traits<char>> &std::operator
-		// <<< std::char_traits<char> > (std::basic_ostream<char, std::char_traits<char>> &,const wchar_t*)': attempting to
-		// reference a deleted function
-		CppStandard = CppStandardVersion.Cpp17;
-
-		// Replace with PCHUsageMode.UseExplicitOrSharedPCHs when this plugin can compile with cpp20
-		PCHUsage = PCHUsageMode.NoPCHs;
+		PCHUsage = PCHUsageMode.UseExplicitOrSharedPCHs;
 
 		PublicDependencyModuleNames.Add("DeveloperSettings");
 
diff --git a/Engine/Plugins/Tests/CQTest/Source/CQTestTests/CQTestTests.Build.cs b/Engine/Plugins/Tests/CQTest/Source/CQTestTests/CQTestTests.Build.cs
--- a/Engine/Plugins/Tests/CQTest/Source/CQTestTests/CQTestTests.Build.cs
+++ b/Engine/Plugins/Tests/CQTest/Source/CQTestTests/CQTestTests.Build.cs
@@ -7,15 +7,7 @@ public class CQTestTests : ModuleRules
 	public CQTestTests(ReadOnlyTargetRules Target)
 		: base(Target)
 	{
-		// Does not compile with C++20:
-		// error C2088: '<<': illegal for class
-		// error C2280: 'std::basic_ostream<char,std::char_traits<char>> &std::operator
-		// <<< std::char_traits<char> > (std::basic_ostream<char, std::char_traits<char>> &,const wchar_t*)': attempting to
-		// reference a deleted function
-		CppStandard = CppStandardVersion.Cpp17;
-
-		// Replace with PCHUsageMode.UseExplicitOrSharedPCHs when this plugin can compile with cpp20
-		PCHUsage = PCHUsageMode.NoPCHs;
+		PCHUsage = PCHUsageMode.UseExplicitOrSharedPCHs;
 
 		PrivateDependencyModuleNames.AddRange(
 			new string[] {
