Fix inconsistency of returned value in PIENetworkComponent methods

#rnx

diff --git a/Engine/Plugins/Tests/CQTest/Source/CQTest/Public/Impl/PIENetworkComponent.inl b/Engine/Plugins/Tests/CQTest/Source/CQTest/Public/Impl/PIENetworkComponent.inl
--- a/Engine/Plugins/Tests/CQTest/Source/CQTest/Public/Impl/PIENetworkComponent.inl
+++ b/Engine/Plugins/Tests/CQTest/Source/CQTest/Public/Impl/PIENetworkComponent.inl
@@ -173,6 +173,7 @@ template<typename NetworkDataType>
 inline FNetworkComponentBuilder<NetworkDataType>& FNetworkComponentBuilder<NetworkDataType>::WithPacketSimulationSettings(FPacketSimulationSettings* InPacketSimulationSettings)
 {
 	PacketSimulationSettings = InPacketSimulationSettings;
+	return *this;
 }
 
 template <typename NetworkDataType>
