Adjusting millisecond formats for logging

#rnx

diff --git a/Engine/Plugins/Tests/CQTest/Source/CQTest/Private/Commands/TestCommands.cpp b/Engine/Plugins/Tests/CQTest/Source/CQTest/Private/Commands/TestCommands.cpp
--- a/Engine/Plugins/Tests/CQTest/Source/CQTest/Private/Commands/TestCommands.cpp
+++ b/Engine/Plugins/Tests/CQTest/Source/CQTest/Private/Commands/TestCommands.cpp
@@ -26,7 +26,7 @@ bool CQTEST_API FWaitUntil::Update()
 	{
 		if (Description)
 		{
-			UE_LOG(LogCqTest, Log, TEXT("Finished %s after %d milliseconds"), Description, int(Elapsed.GetTotalMilliseconds()));
+			UE_LOG(LogCqTest, Log, TEXT("Finished %s after %.0f milliseconds"), Description, Elapsed.GetTotalMilliseconds());
 		}
 		return true;
 	}
@@ -34,11 +34,11 @@ bool CQTEST_API FWaitUntil::Update()
 	{
 		if (Description)
 		{
-			TestRunner.AddError(*FString::Printf(TEXT("Timed out waiting for %s after %d milliseconds"), Description, Elapsed.GetTotalMilliseconds()));
+			TestRunner.AddError(*FString::Printf(TEXT("Timed out waiting for %s after %.0f milliseconds"), Description, Elapsed.GetTotalMilliseconds()));
 		}
 		else
 		{
-			TestRunner.AddError(*FString::Printf(TEXT("Latent command timed out after %d milliseconds"), Elapsed.GetTotalMilliseconds()));
+			TestRunner.AddError(*FString::Printf(TEXT("Latent command timed out after %.0f milliseconds"), Elapsed.GetTotalMilliseconds()));
 		}
 		return true;
 	}
diff --git a/Engine/Plugins/Tests/CQTest/Source/CQTestTests/Private/CQTestTests.cpp b/Engine/Plugins/Tests/CQTest/Source/CQTestTests/Private/CQTestTests.cpp
--- a/Engine/Plugins/Tests/CQTest/Source/CQTestTests/Private/CQTestTests.cpp
+++ b/Engine/Plugins/Tests/CQTest/Source/CQTestTests/Private/CQTestTests.cpp
@@ -373,6 +373,14 @@ namespace CQTestTests
 				Assert.Fail("Unexpected Error");
 			});
 		}
+
+		TEST_METHOD(Timeout_InLatentActions_ProvidesTimeoutInErrorMessage)
+		{
+			Assert.ExpectErrorRegex(TEXT("\\d{2,} milliseconds"));
+			TestCommandBuilder.StartWhen([]() {
+				return false;
+				}, FTimespan::FromMilliseconds(50));
+		}
 	};
 
 	// --------------------------------------------------------
