Added new test against the FTestMapSpawner

#rnx

diff --git a/Engine/Plugins/Tests/CQTest/Source/CQTest/Private/Components/MapTestSpawner.cpp b/Engine/Plugins/Tests/CQTest/Source/CQTest/Private/Components/MapTestSpawner.cpp
--- a/Engine/Plugins/Tests/CQTest/Source/CQTest/Private/Components/MapTestSpawner.cpp
+++ b/Engine/Plugins/Tests/CQTest/Source/CQTest/Private/Components/MapTestSpawner.cpp
@@ -8,6 +8,67 @@
 #include "Engine/Engine.h"
 #include "Misc/Paths.h"
 
+#if WITH_EDITOR
+#include "Editor.h"
+#include "HAL/FileManager.h"
+#include "LevelEditorSubsystem.h"
+#include "Tests/AutomationEditorCommon.h"
+#endif
+
+namespace {
+
+static const FString TempMapDirectory = FPaths::Combine(FPaths::ProjectContentDir(), TEXT("CQTestMapTemp"));
+
+/**
+ * Generates a unique random 8 character map name.
+ */
+FString GenerateUniqueMapName()
+{
+	FString UniqueMapName = FGuid::NewGuid().ToString();
+	UniqueMapName.LeftInline(8);
+
+	return UniqueMapName;
+}
+
+/**
+ * Cleans up all created resources.
+ */
+void CleanupTempResources()
+{
+#if WITH_EDITOR
+	bool bDirectoryMustExist = true;
+	bool bRemoveRecursively = true;
+	bool bWasDeleted = IFileManager::Get().DeleteDirectory(*TempMapDirectory, bDirectoryMustExist, bRemoveRecursively);
+	check(bWasDeleted);
+#endif
+}
+
+} //anonymous
+
+TUniquePtr<FMapTestSpawner> FMapTestSpawner::CreateFromTempLevel(FTestCommandBuilder& InCommandBuilder)
+{
+#if WITH_EDITOR
+	FString MapName = GenerateUniqueMapName();
+	FString MapPath = FPaths::Combine(TempMapDirectory, MapName);
+	FString NewLevelPackage = FPackageName::FilenameToLongPackageName(MapPath);
+
+	ULevelEditorSubsystem* LevelEditorSubsystem = GEditor->GetEditorSubsystem<ULevelEditorSubsystem>();
+	bool bWasTempLevelCreated = LevelEditorSubsystem->NewLevel(NewLevelPackage);
+	check(bWasTempLevelCreated);
+
+	TUniquePtr<FMapTestSpawner> Spawner = MakeUnique<FMapTestSpawner>(TempMapDirectory, MapName);
+	InCommandBuilder.OnTearDown([&]() {
+		// Create a new map to free up the reference to the map used during testing before cleaning up all temporary resources
+		FAutomationEditorCommonUtils::CreateNewMap();
+		CleanupTempResources();
+	});
+	return MoveTemp(Spawner);
+#else
+	checkf(false, TEXT("CreateFromTempLevel can't create a new level if WITH_EDITOR=false"));
+	return nullptr;
+#endif
+}
+
 void FMapTestSpawner::AddWaitUntilLoadedCommand(FAutomationTestBase* TestRunner)
 {
 #if WITH_AUTOMATION_TESTS
diff --git a/Engine/Plugins/Tests/CQTest/Source/CQTest/Public/Components/MapTestSpawner.h b/Engine/Plugins/Tests/CQTest/Source/CQTest/Public/Components/MapTestSpawner.h
--- a/Engine/Plugins/Tests/CQTest/Source/CQTest/Public/Components/MapTestSpawner.h
+++ b/Engine/Plugins/Tests/CQTest/Source/CQTest/Public/Components/MapTestSpawner.h
@@ -4,16 +4,41 @@
 
 #include "SpawnHelper.h"
 
+#include "Commands/TestCommandBuilder.h"
+
 /// Class for spawning Actors in a named Map/Level
 struct CQTEST_API FMapTestSpawner : public FSpawnHelper
 {
+	/**
+	 * Construct the MapTestSpawner.
+	 *
+	 * @param MapDirectory - The directory which the map resides in.
+	 * @param MapName - Name of the map.
+	 */
 	FMapTestSpawner(const FString& MapDirectory, const FString& MapName)
 		: MapDirectory(MapDirectory)
 		, MapName(MapName)
 	{
 	}
 
+	/**
+	 * Creates an instance of the MapTestSpawner with a temporary level ready for use.
+	 * 
+	 * @param InCommandBuilder - Test Command Builder used to assist with setup.
+	 * @return unique instance of the FMapTestSpawner, nullptr otherwise
+	 */
+	static TUniquePtr<FMapTestSpawner> CreateFromTempLevel(FTestCommandBuilder& InCommandBuilder);
+
+	/**
+	 * Loads the map specified from the MapDirectory and MapName to be prepared for the test.
+	 *
+	 * @param TestRunner - TestRunner used to send the latent command needed for map preparations.
+	 */
 	void AddWaitUntilLoadedCommand(FAutomationTestBase* TestRunner);
+
+	/**
+	 * Finds the first pawn in the given map.
+	 */
 	APawn* FindFirstPlayerPawn();
 
 protected:
diff --git a/Engine/Plugins/Tests/CQTest/Source/CQTestTests/CQTestTests.Build.cs b/Engine/Plugins/Tests/CQTest/Source/CQTestTests/CQTestTests.Build.cs
--- a/Engine/Plugins/Tests/CQTest/Source/CQTestTests/CQTestTests.Build.cs
+++ b/Engine/Plugins/Tests/CQTest/Source/CQTestTests/CQTestTests.Build.cs
@@ -17,5 +17,15 @@ public class CQTestTests : ModuleRules
 					"CQTest"
 				 }
 			);
+
+		if (Target.bBuildEditor)
+		{
+			PrivateDependencyModuleNames.AddRange(
+				new string[] {
+					"EngineSettings",
+					"LevelEditor",
+					"UnrealEd"
+			});
+		}
 	}
 }
diff --git a/Engine/Plugins/Tests/CQTest/Source/CQTestTests/Private/Components/MapSpawnHelperTests.cpp b/Engine/Plugins/Tests/CQTest/Source/CQTestTests/Private/Components/MapSpawnHelperTests.cpp
new file mode 100644
--- /dev/null
+++ b/Engine/Plugins/Tests/CQTest/Source/CQTestTests/Private/Components/MapSpawnHelperTests.cpp
@@ -0,0 +1,35 @@
+// Copyright Epic Games, Inc. All Rights Reserved.
+
+#define ENABLE_MAPSPAWNER_TEST WITH_EDITOR && (WITH_DEV_AUTOMATION_TESTS || WITH_PERF_AUTOMATION_TESTS)
+
+#if ENABLE_MAPSPAWNER_TEST
+
+#include "CQTest.h"
+#include "Components/MapTestSpawner.h"
+#include "GameFramework/Pawn.h"
+
+TEST_CLASS(MapSpawnHelperTests, "TestFramework.CQTest.Map")
+{
+	TUniquePtr<FMapTestSpawner> Spawner;
+
+	BEFORE_EACH() {
+		Spawner = FMapTestSpawner::CreateFromTempLevel(TestCommandBuilder);
+		ASSERT_THAT(IsNotNull(Spawner));
+
+		Spawner->AddWaitUntilLoadedCommand(TestRunner);
+		TestCommandBuilder.Do([&]() {
+			// Because we're creating a level for this test, we will also want to populate the level with a Pawn object that can be then tested against
+			Spawner->SpawnActor<APawn>();
+		});
+	}
+
+	TEST_METHOD(MapSpawner_FindsPlayerSpawn)
+	{
+		TestCommandBuilder.Do([&]() {
+			APawn* Player = Spawner->FindFirstPlayerPawn();
+			ASSERT_THAT(IsNotNull(Player));
+		});
+	}
+};
+
+#endif // ENABLE_MAPSPAWNER_TEST
\ No newline at end of file
