Addressing potential race condition where a new map was created while a PIE session was active

#rnx

diff --git a/Engine/Plugins/Tests/CQTest/Source/CQTest/Private/Components/PIENetworkComponent.cpp b/Engine/Plugins/Tests/CQTest/Source/CQTest/Private/Components/PIENetworkComponent.cpp
--- a/Engine/Plugins/Tests/CQTest/Source/CQTest/Private/Components/PIENetworkComponent.cpp
+++ b/Engine/Plugins/Tests/CQTest/Source/CQTest/Private/Components/PIENetworkComponent.cpp
@@ -28,7 +28,9 @@ FBasePIENetworkComponent::FBasePIENetworkComponent(FAutomationTestBase* InTestRu
 	}
 
 	CommandBuilder
-		->Do(TEXT("Start PIE"), [this]() { StartPie(); })
+		->Do(TEXT("Stop PIE"), [this]() { StopPie(); })
+		.Then(TEXT("Create New Map"), [this]() { FAutomationEditorCommonUtils::CreateNewMap(); })
+		.Then(TEXT("Start PIE"), [this]() { StartPie(); })
 		.Until(TEXT("Collect PIE Worlds"), [this]() { return CollectPieWorlds(); })
 		.Until(TEXT("Await Connections"), [this]() { return AwaitConnections(); })
 		.OnTearDown(TEXT("Restore Editor State"), [this]() { RestoreState(); });
@@ -79,7 +81,7 @@ FBasePIENetworkComponent& FBasePIENetworkComponent::StartWhen(TCHAR* Description
 	return *this;
 }
 
-void FBasePIENetworkComponent::StartPie()
+void FBasePIENetworkComponent::StopPie()
 {
 	if (ServerState == nullptr)
 	{
@@ -87,8 +89,10 @@ void FBasePIENetworkComponent::StartPie()
 		return;
 	}
 	GUnrealEd->RequestEndPlayMap();
-	FAutomationEditorCommonUtils::CreateNewMap();
+}
 
+void FBasePIENetworkComponent::StartPie()
+{
 	ULevelEditorPlaySettings* PlaySettings = NewObject<ULevelEditorPlaySettings>();
 	if (ServerState->bIsDedicatedServer)
 	{
diff --git a/Engine/Plugins/Tests/CQTest/Source/CQTest/Public/Components/PIENetworkComponent.h b/Engine/Plugins/Tests/CQTest/Source/CQTest/Public/Components/PIENetworkComponent.h
--- a/Engine/Plugins/Tests/CQTest/Source/CQTest/Public/Components/PIENetworkComponent.h
+++ b/Engine/Plugins/Tests/CQTest/Source/CQTest/Public/Components/PIENetworkComponent.h
@@ -40,6 +40,7 @@ public:
 	FBasePIENetworkComponent& StartWhen(TCHAR* Description, TFunction<bool()> Query, FTimespan Timeout = FTimespan::FromSeconds(10));
 
 protected:
+	void StopPie();
 	void StartPie();
 	bool CollectPieWorlds();
 	bool AwaitConnections();
