Fixing CookedEditor compilation error by missing a reference to `FAutomationEditorCommonUtils`.

#rnx

diff --git a/Engine/Plugins/Tests/CQTest/Source/CQTest/Private/Components/MapTestSpawner.cpp b/Engine/Plugins/Tests/CQTest/Source/CQTest/Private/Components/MapTestSpawner.cpp
--- a/Engine/Plugins/Tests/CQTest/Source/CQTest/Private/Components/MapTestSpawner.cpp
+++ b/Engine/Plugins/Tests/CQTest/Source/CQTest/Private/Components/MapTestSpawner.cpp
@@ -1,19 +1,17 @@
 // Copyright Epic Games, Inc. All Rights Reserved.
 
 #include "Components/MapTestSpawner.h"
-#include "Commands/TestCommands.h"
 
+#if ENABLE_MAPSPAWNER_TEST
+#include "Commands/TestCommands.h"
 #include "Tests/AutomationCommon.h"
 #include "GameFramework/PlayerController.h"
 #include "Engine/Engine.h"
 #include "Misc/Paths.h"
-
-#if WITH_EDITOR
 #include "Editor.h"
 #include "HAL/FileManager.h"
 #include "LevelEditorSubsystem.h"
 #include "Tests/AutomationEditorCommon.h"
-#endif
 
 namespace {
 
@@ -35,19 +33,16 @@ FString GenerateUniqueMapName()
  */
 void CleanupTempResources()
 {
-#if WITH_EDITOR
 	bool bDirectoryMustExist = true;
 	bool bRemoveRecursively = true;
 	bool bWasDeleted = IFileManager::Get().DeleteDirectory(*TempMapDirectory, bDirectoryMustExist, bRemoveRecursively);
 	check(bWasDeleted);
-#endif
 }
 
 } //anonymous
 
 TUniquePtr<FMapTestSpawner> FMapTestSpawner::CreateFromTempLevel(FTestCommandBuilder& InCommandBuilder)
 {
-#if WITH_EDITOR
 	FString MapName = GenerateUniqueMapName();
 	FString MapPath = FPaths::Combine(TempMapDirectory, MapName);
 	FString NewLevelPackage = FPackageName::FilenameToLongPackageName(MapPath);
@@ -63,15 +58,10 @@ TUniquePtr<FMapTestSpawner> FMapTestSpawner::CreateFromTempLevel(FTestCommandBui
 		CleanupTempResources();
 	});
 	return MoveTemp(Spawner);
-#else
-	checkf(false, TEXT("CreateFromTempLevel can't create a new level if WITH_EDITOR=false"));
-	return nullptr;
-#endif
 }
 
 void FMapTestSpawner::AddWaitUntilLoadedCommand(FAutomationTestBase* TestRunner)
 {
-#if WITH_AUTOMATION_TESTS
 	check(PieWorld == nullptr);
 
 	const FString FileName = FString::Printf(TEXT("%s.%s"), *MapName, *MapName);
@@ -91,9 +81,6 @@ void FMapTestSpawner::AddWaitUntilLoadedCommand(FAutomationTestBase* TestRunner)
 
 		return false;
 	}));
-#else
-	checkf(false, TEXT("AddWaitUntilLoadedCommand can't call AutomationOpenMap if WITH_AUTOMATION_TESTS=false"));
-#endif
 }
 
 UWorld* FMapTestSpawner::CreateWorld()
@@ -106,3 +93,5 @@ APawn* FMapTestSpawner::FindFirstPlayerPawn()
 {
 	return GetWorld().GetFirstPlayerController()->GetPawn();
 }
+
+#endif // ENABLE_MAPSPAWNER_TEST
\ No newline at end of file
diff --git a/Engine/Plugins/Tests/CQTest/Source/CQTest/Public/Components/MapTestSpawner.h b/Engine/Plugins/Tests/CQTest/Source/CQTest/Public/Components/MapTestSpawner.h
--- a/Engine/Plugins/Tests/CQTest/Source/CQTest/Public/Components/MapTestSpawner.h
+++ b/Engine/Plugins/Tests/CQTest/Source/CQTest/Public/Components/MapTestSpawner.h
@@ -6,6 +6,10 @@
 
 #include "Commands/TestCommandBuilder.h"
 
+#define ENABLE_MAPSPAWNER_TEST WITH_EDITOR && WITH_AUTOMATION_TESTS
+
+#if ENABLE_MAPSPAWNER_TEST
+
 /// Class for spawning Actors in a named Map/Level
 struct CQTEST_API FMapTestSpawner : public FSpawnHelper
 {
@@ -48,4 +52,6 @@ private:
 	FString MapDirectory;
 	FString MapName;
 	UWorld* PieWorld{ nullptr };
-};
\ No newline at end of file
+};
+
+#endif // ENABLE_MAPSPAWNER_TEST
\ No newline at end of file
diff --git a/Engine/Plugins/Tests/CQTest/Source/CQTestTests/Private/Components/MapSpawnHelperTests.cpp b/Engine/Plugins/Tests/CQTest/Source/CQTestTests/Private/Components/MapSpawnHelperTests.cpp
--- a/Engine/Plugins/Tests/CQTest/Source/CQTestTests/Private/Components/MapSpawnHelperTests.cpp
+++ b/Engine/Plugins/Tests/CQTest/Source/CQTestTests/Private/Components/MapSpawnHelperTests.cpp
@@ -1,13 +1,11 @@
 // Copyright Epic Games, Inc. All Rights Reserved.
 
-#define ENABLE_MAPSPAWNER_TEST WITH_EDITOR && (WITH_DEV_AUTOMATION_TESTS || WITH_PERF_AUTOMATION_TESTS)
-
-#if ENABLE_MAPSPAWNER_TEST
-
 #include "CQTest.h"
 #include "Components/MapTestSpawner.h"
 #include "GameFramework/Pawn.h"
 
+#if ENABLE_MAPSPAWNER_TEST
+
 TEST_CLASS(MapSpawnHelperTests, "TestFramework.CQTest.Map")
 {
 	TUniquePtr<FMapTestSpawner> Spawner;
