Fixing heap use after free caught by running tests under address sanitization

#rnx

diff --git a/Engine/Plugins/Tests/CQTest/Source/CQTest/Private/Components/MapTestSpawner.cpp b/Engine/Plugins/Tests/CQTest/Source/CQTest/Private/Components/MapTestSpawner.cpp
--- a/Engine/Plugins/Tests/CQTest/Source/CQTest/Private/Components/MapTestSpawner.cpp
+++ b/Engine/Plugins/Tests/CQTest/Source/CQTest/Private/Components/MapTestSpawner.cpp
@@ -10,9 +10,12 @@
 #include "Misc/Paths.h"
 #include "Editor.h"
 #include "HAL/FileManager.h"
+#include "LevelEditor.h"
 #include "LevelEditorSubsystem.h"
 #include "Tests/AutomationEditorCommon.h"
 
+DEFINE_LOG_CATEGORY_STATIC(LogMapTest, Log, All);
+
 namespace {
 
 static const FString TempMapDirectory = FPaths::Combine(FPaths::ProjectContentDir(), TEXT("CQTestMapTemp"));
@@ -41,6 +44,13 @@ void CleanupTempResources()
 
 } //anonymous
 
+FMapTestSpawner::FMapTestSpawner(const FString& MapDirectory, const FString& MapName) : MapDirectory(MapDirectory) , MapName(MapName)
+{
+	// Register Map Change Events
+	FLevelEditorModule& LevelEditor = FModuleManager::LoadModuleChecked<FLevelEditorModule>("LevelEditor");
+	MapChangedHandle = LevelEditor.OnMapChanged().AddRaw(this, &FMapTestSpawner::OnMapChanged);
+}
+
 TUniquePtr<FMapTestSpawner> FMapTestSpawner::CreateFromTempLevel(FTestCommandBuilder& InCommandBuilder)
 {
 	FString MapName = GenerateUniqueMapName();
@@ -94,4 +104,17 @@ APawn* FMapTestSpawner::FindFirstPlayerPawn()
 	return GetWorld().GetFirstPlayerController()->GetPawn();
 }
 
+void FMapTestSpawner::OnMapChanged(UWorld* World, EMapChangeType ChangeType)
+{
+	if (PieWorld && ChangeType == EMapChangeType::TearDownWorld)
+	{
+		UE_LOG(LogMapTest, Verbose, TEXT("Map used by the Spawner has been changed."));
+		GameWorld = nullptr;
+		PieWorld = nullptr;
+
+		FLevelEditorModule& LevelEditor = FModuleManager::LoadModuleChecked<FLevelEditorModule>("LevelEditor");
+		LevelEditor.OnMapChanged().Remove(MapChangedHandle);
+	}
+}
+
 #endif // ENABLE_MAPSPAWNER_TEST
\ No newline at end of file
diff --git a/Engine/Plugins/Tests/CQTest/Source/CQTest/Public/Components/MapTestSpawner.h b/Engine/Plugins/Tests/CQTest/Source/CQTest/Public/Components/MapTestSpawner.h
--- a/Engine/Plugins/Tests/CQTest/Source/CQTest/Public/Components/MapTestSpawner.h
+++ b/Engine/Plugins/Tests/CQTest/Source/CQTest/Public/Components/MapTestSpawner.h
@@ -9,6 +9,7 @@
 #define ENABLE_MAPSPAWNER_TEST WITH_EDITOR && WITH_AUTOMATION_TESTS
 
 #if ENABLE_MAPSPAWNER_TEST
+#include "UnrealEdMisc.h"
 
 /// Class for spawning Actors in a named Map/Level
 struct CQTEST_API FMapTestSpawner : public FSpawnHelper
@@ -19,11 +20,7 @@ struct CQTEST_API FMapTestSpawner : public FSpawnHelper
 	 * @param MapDirectory - The directory which the map resides in.
 	 * @param MapName - Name of the map.
 	 */
-	FMapTestSpawner(const FString& MapDirectory, const FString& MapName)
-		: MapDirectory(MapDirectory)
-		, MapName(MapName)
-	{
-	}
+	FMapTestSpawner(const FString& MapDirectory, const FString& MapName);
 
 	/**
 	 * Creates an instance of the MapTestSpawner with a temporary level ready for use.
@@ -49,9 +46,15 @@ protected:
     virtual UWorld* CreateWorld() override;
 
 private:
+	/**
+	 * Handler called on map changed.
+	 */
+	void OnMapChanged(UWorld* World, EMapChangeType ChangeType);
+
 	FString MapDirectory;
 	FString MapName;
 	UWorld* PieWorld{ nullptr };
+	FDelegateHandle MapChangedHandle;
 };
 
 #endif // ENABLE_MAPSPAWNER_TEST
\ No newline at end of file
