Replaced UE_NODISCARD with [[nodiscard]].

diff --git a/Engine/Plugins/Editor/ObjectMixer/ObjectMixer/Source/ObjectMixer/Public/Views/List/ObjectMixerEditorList.h b/Engine/Plugins/Editor/ObjectMixer/ObjectMixer/Source/ObjectMixer/Public/Views/List/ObjectMixerEditorList.h
--- a/Engine/Plugins/Editor/ObjectMixer/ObjectMixer/Source/ObjectMixer/Public/Views/List/ObjectMixerEditorList.h
+++ b/Engine/Plugins/Editor/ObjectMixer/ObjectMixer/Source/ObjectMixer/Public/Views/List/ObjectMixerEditorList.h
@@ -158,7 +158,7 @@ public:
 
 	void FlushWidget();
 
-	UE_NODISCARD TArray<TSharedPtr<ISceneOutlinerTreeItem>> GetSelectedTreeViewItems() const;
+	[[nodiscard]] TArray<TSharedPtr<ISceneOutlinerTreeItem>> GetSelectedTreeViewItems() const;
 	int32 GetSelectedTreeViewItemCount() const;
 
 	TSet<TSharedPtr<ISceneOutlinerTreeItem>> GetSoloRows() const;
@@ -174,8 +174,8 @@ public:
 	void EvaluateAndSetEditorVisibilityPerRow();
 	
 	/** Represents the collections the user has selected in the UI. If empty, "All" is considered as selected. */
-	UE_NODISCARD const TSet<FName>& GetSelectedCollections() const;
-	UE_NODISCARD bool IsCollectionSelected(const FName& CollectionName) const;
+	[[nodiscard]] const TSet<FName>& GetSelectedCollections() const;
+	[[nodiscard]] bool IsCollectionSelected(const FName& CollectionName) const;
 	void SetSelectedCollections(const TSet<FName> InSelectedCollections);
 	void SetCollectionSelected(const FName& CollectionName, const bool bNewSelected);
 
diff --git a/Engine/Plugins/Editor/ObjectMixer/ObjectMixer/Source/ObjectMixer/Public/Views/List/ObjectMixerUtils.h b/Engine/Plugins/Editor/ObjectMixer/ObjectMixer/Source/ObjectMixer/Public/Views/List/ObjectMixerUtils.h
--- a/Engine/Plugins/Editor/ObjectMixer/ObjectMixer/Source/ObjectMixer/Public/Views/List/ObjectMixerUtils.h
+++ b/Engine/Plugins/Editor/ObjectMixer/ObjectMixer/Source/ObjectMixer/Public/Views/List/ObjectMixerUtils.h
@@ -30,7 +30,7 @@ namespace FObjectMixerUtils
 
 	UObject* GetRowObject(TSharedPtr<ISceneOutlinerTreeItem> InTreeItem, const bool bGetHybridRowComponent = false);
 	
-	UE_NODISCARD AActor* GetSelfOrOuterAsActor(TSharedPtr<ISceneOutlinerTreeItem> InTreeItem);
+	[[nodiscard]] AActor* GetSelfOrOuterAsActor(TSharedPtr<ISceneOutlinerTreeItem> InTreeItem);
 
 	bool IsObjectRefInCollection(const FName& CollectionName, const UObject* Object, const TSharedPtr<FObjectMixerEditorList> ListModel);
 	bool IsObjectRefInCollection(const FName& CollectionName, TSharedPtr<ISceneOutlinerTreeItem> InTreeItem);
diff --git a/Engine/Plugins/Editor/ObjectMixer/ObjectMixer/Source/ObjectMixer/Public/Views/List/SObjectMixerEditorList.h b/Engine/Plugins/Editor/ObjectMixer/ObjectMixer/Source/ObjectMixer/Public/Views/List/SObjectMixerEditorList.h
--- a/Engine/Plugins/Editor/ObjectMixer/ObjectMixer/Source/ObjectMixer/Public/Views/List/SObjectMixerEditorList.h
+++ b/Engine/Plugins/Editor/ObjectMixer/ObjectMixer/Source/ObjectMixer/Public/Views/List/SObjectMixerEditorList.h
@@ -93,7 +93,7 @@ public:
 	void SetTreeViewMode(EObjectMixerTreeViewMode InViewMode);
 
 	const TSet<FName>& GetSelectedCollections();
-	UE_NODISCARD bool IsCollectionSelected(const FName& CollectionName);
+	[[nodiscard]] bool IsCollectionSelected(const FName& CollectionName);
 	
 	void RebuildCollectionSelector();
 
@@ -129,7 +129,7 @@ public:
 
 	void AddToPendingPropertyPropagations(const FObjectMixerEditorListRowData::FPropertyPropagationInfo& InPropagationInfo);
 
-	UE_NODISCARD TArray<TSharedPtr<ISceneOutlinerTreeItem>> GetSelectedTreeViewItems() const;
+	[[nodiscard]] TArray<TSharedPtr<ISceneOutlinerTreeItem>> GetSelectedTreeViewItems() const;
 	int32 GetSelectedTreeViewItemCount() const;
 
 	void SetTreeViewItemSelected(TSharedRef<ISceneOutlinerTreeItem> Item, const bool bNewSelected);
diff --git a/Engine/Plugins/Experimental/MeshLODToolset/Source/MeshLODToolset/Public/Graphs/GenerateStaticMeshLODProcess.h b/Engine/Plugins/Experimental/MeshLODToolset/Source/MeshLODToolset/Public/Graphs/GenerateStaticMeshLODProcess.h
--- a/Engine/Plugins/Experimental/MeshLODToolset/Source/MeshLODToolset/Public/Graphs/GenerateStaticMeshLODProcess.h
+++ b/Engine/Plugins/Experimental/MeshLODToolset/Source/MeshLODToolset/Public/Graphs/GenerateStaticMeshLODProcess.h
@@ -382,7 +382,7 @@ public:
 	 *  - Read all the Source Texture Images that need to be baked
 	 *  - call InitializeGenerator() to set up the compute graph
 	 */
-	UE_NODISCARD bool Initialize(UStaticMesh* SourceMesh, FProgressCancel* Progress = nullptr);
+	[[nodiscard]] bool Initialize(UStaticMesh* SourceMesh, FProgressCancel* Progress = nullptr);
 
 
 	/** @return the source UStaticMesh passed to Initialize() */
diff --git a/Engine/Plugins/Experimental/MeshModelingToolsetExp/Source/MeshModelingToolsExp/Private/DrawPolyPathTool.cpp b/Engine/Plugins/Experimental/MeshModelingToolsetExp/Source/MeshModelingToolsExp/Private/DrawPolyPathTool.cpp
--- a/Engine/Plugins/Experimental/MeshModelingToolsetExp/Source/MeshModelingToolsExp/Private/DrawPolyPathTool.cpp
+++ b/Engine/Plugins/Experimental/MeshModelingToolsetExp/Source/MeshModelingToolsExp/Private/DrawPolyPathTool.cpp
@@ -45,7 +45,7 @@ namespace DrawPolyPathToolLocals
 
 	/// Generate path mesh
 	/// @return Offset to location of generated mesh 
-	UE_NODISCARD FVector3d GeneratePathMesh(FDynamicMesh3& Mesh,
+	[[nodiscard]] FVector3d GeneratePathMesh(FDynamicMesh3& Mesh,
 		const TArray<FFrame3d>& InPathPoints, 
 		const TArray<double>& InOffsetScaleFactors, 
 		double OffsetDistance, 
diff --git a/Engine/Plugins/Experimental/MeshModelingToolsetExp/Source/MeshModelingToolsExp/Public/DrawPolyPathTool.h b/Engine/Plugins/Experimental/MeshModelingToolsetExp/Source/MeshModelingToolsExp/Public/DrawPolyPathTool.h
--- a/Engine/Plugins/Experimental/MeshModelingToolsetExp/Source/MeshModelingToolsExp/Public/DrawPolyPathTool.h
+++ b/Engine/Plugins/Experimental/MeshModelingToolsetExp/Source/MeshModelingToolsExp/Public/DrawPolyPathTool.h
@@ -285,7 +285,7 @@ protected:
 	void InitializePreviewMesh();
 	void ClearPreview();
 
-	UE_NODISCARD FVector3d GeneratePathMesh(UE::Geometry::FDynamicMesh3& Mesh);
+	[[nodiscard]] FVector3d GeneratePathMesh(UE::Geometry::FDynamicMesh3& Mesh);
 	void GenerateExtrudeMesh(UE::Geometry::FDynamicMesh3& PathMesh);
 	void EmitNewObject();
 
diff --git a/Engine/Plugins/Experimental/Mutable/Source/MutableRuntime/Internal/MuR/SystemPrivate.h b/Engine/Plugins/Experimental/Mutable/Source/MutableRuntime/Internal/MuR/SystemPrivate.h
--- a/Engine/Plugins/Experimental/Mutable/Source/MutableRuntime/Internal/MuR/SystemPrivate.h
+++ b/Engine/Plugins/Experimental/Mutable/Source/MutableRuntime/Internal/MuR/SystemPrivate.h
@@ -1215,7 +1215,7 @@ namespace mu
         void MarkRomUsed( int32 RomIndex, const TSharedPtr<const Model>& );
 
 		/** */
-		UE_NODISCARD Ptr<Image> CreateImage(uint32 SizeX, uint32 SizeY, uint32 Lods, EImageFormat Format, EInitializationType Init)
+		[[nodiscard]] Ptr<Image> CreateImage(uint32 SizeX, uint32 SizeY, uint32 Lods, EImageFormat Format, EInitializationType Init)
 		{
 			CheckRunnerThread();
 
@@ -1260,7 +1260,7 @@ namespace mu
 		}
 
 		/** Ref will be nulled and relesed in any case. */
-		UE_NODISCARD Ptr<Image> CloneOrTakeOver(Ptr<const Image>& Resource)
+		[[nodiscard]] Ptr<Image> CloneOrTakeOver(Ptr<const Image>& Resource)
 		{
 			CheckRunnerThread();
 
@@ -1357,7 +1357,7 @@ namespace mu
 			Resource = nullptr;
 		}
 
-		UE_NODISCARD Ptr<Mesh> CreateMesh(int32 BudgetReserveSize)
+		[[nodiscard]] Ptr<Mesh> CreateMesh(int32 BudgetReserveSize)
 		{
 			CheckRunnerThread();
 
@@ -1370,7 +1370,7 @@ namespace mu
 			return Result;
 		}
 
-		UE_NODISCARD Ptr<Mesh> CloneOrTakeOver(Ptr<const Mesh>& Resource)
+		[[nodiscard]] Ptr<Mesh> CloneOrTakeOver(Ptr<const Mesh>& Resource)
 		{
 			CheckRunnerThread();
 
@@ -1420,7 +1420,7 @@ namespace mu
 		}
 
 		/** */
-		UE_NODISCARD Ptr<const Mesh> LoadMesh(const FCacheAddress& From, bool bTakeOwnership = false)
+		[[nodiscard]] Ptr<const Mesh> LoadMesh(const FCacheAddress& From, bool bTakeOwnership = false)
 		{
 			bool bIsLastReference = false;
 			Ptr<const Mesh> Result = CurrentInstanceCache->GetMesh(From, bIsLastReference);
@@ -1450,7 +1450,7 @@ namespace mu
 		}
 
 		/** */
-		UE_NODISCARD Ptr<const Image> LoadImage(const FCacheAddress& From, bool bTakeOwnership = false)
+		[[nodiscard]] Ptr<const Image> LoadImage(const FCacheAddress& From, bool bTakeOwnership = false)
 		{
 			bool bIsLastReference = false;
 			Ptr<const Image> Result = CurrentInstanceCache->GetImage(From, bIsLastReference);
diff --git a/Engine/Plugins/Experimental/Mutable/Source/MutableRuntime/Private/MuR/CodeRunner.h b/Engine/Plugins/Experimental/Mutable/Source/MutableRuntime/Private/MuR/CodeRunner.h
--- a/Engine/Plugins/Experimental/Mutable/Source/MutableRuntime/Private/MuR/CodeRunner.h
+++ b/Engine/Plugins/Experimental/Mutable/Source/MutableRuntime/Private/MuR/CodeRunner.h
@@ -590,12 +590,12 @@ namespace  mu
 			return m_pSystem->WorkingMemoryManager.Release(Resource);
 		}
 
-		UE_NODISCARD inline Ptr<Mesh> CreateMesh(int32 BudgetReserveSize = 0)
+		[[nodiscard]] inline Ptr<Mesh> CreateMesh(int32 BudgetReserveSize = 0)
 		{
 			return m_pSystem->WorkingMemoryManager.CreateMesh(BudgetReserveSize);
 		}
 
-		UE_NODISCARD inline Ptr<Mesh> CloneOrTakeOver(Ptr<const Mesh>& Ref)
+		[[nodiscard]] inline Ptr<Mesh> CloneOrTakeOver(Ptr<const Mesh>& Ref)
 		{
 			return m_pSystem->WorkingMemoryManager.CloneOrTakeOver(Ref);
 		}
diff --git a/Engine/Plugins/Experimental/PCG/Source/PCG/Public/Helpers/PCGMetadataHelpers.h b/Engine/Plugins/Experimental/PCG/Source/PCG/Public/Helpers/PCGMetadataHelpers.h
--- a/Engine/Plugins/Experimental/PCG/Source/PCG/Public/Helpers/PCGMetadataHelpers.h
+++ b/Engine/Plugins/Experimental/PCG/Source/PCG/Public/Helpers/PCGMetadataHelpers.h
@@ -23,10 +23,10 @@ namespace PCGMetadataHelpers
 	PCG_API UPCGMetadata* GetMutableMetadata(UPCGData* InData);
 	
 	/** Create a lambda that will construct a soft object path from an underlying attribute of type FSoftObjectPath or FString. Returns true if successful. */
-	UE_NODISCARD PCG_API bool CreateObjectPathGetter(const FPCGMetadataAttributeBase* InAttributeBase, TFunction<void(int64, FSoftObjectPath&)>& OutGetter);
+	[[nodiscard]] PCG_API bool CreateObjectPathGetter(const FPCGMetadataAttributeBase* InAttributeBase, TFunction<void(int64, FSoftObjectPath&)>& OutGetter);
 
 	/** Create a lambda that will construct a soft object path from an underlying attribute of type FSoftObjectPath or FString. Returns true if successful. */
-	UE_NODISCARD PCG_API bool CreateObjectOrClassPathGetter(const FPCGMetadataAttributeBase* InAttributeBase, TFunction<void(int64, FSoftObjectPath&)>& OutGetter);
+	[[nodiscard]] PCG_API bool CreateObjectOrClassPathGetter(const FPCGMetadataAttributeBase* InAttributeBase, TFunction<void(int64, FSoftObjectPath&)>& OutGetter);
 
 	/** Copy the attribute coming from Source Data into Target Data.
 	* Can set bSameOrigin to true if SourceData and TargetData have the same origin (if TargetData was initialized from SourceData). 
diff --git a/Engine/Plugins/Experimental/PCG/Source/PCG/Public/PCGInputOutputSettings.h b/Engine/Plugins/Experimental/PCG/Source/PCG/Public/PCGInputOutputSettings.h
--- a/Engine/Plugins/Experimental/PCG/Source/PCG/Public/PCGInputOutputSettings.h
+++ b/Engine/Plugins/Experimental/PCG/Source/PCG/Public/PCGInputOutputSettings.h
@@ -53,7 +53,7 @@ public:
 	// Add a new custom pin
 	// Note that you should use the return value of this function, since it can be different from
 	// the one passed as argument. It will change if its label collides with existing pins.
-	UE_NODISCARD const FPCGPinProperties& AddPin(const FPCGPinProperties& NewCustomPinProperties);
+	[[nodiscard]] const FPCGPinProperties& AddPin(const FPCGPinProperties& NewCustomPinProperties);
 
 protected:
 	virtual FPCGElementPtr CreateElement() const override { return MakeShared<FPCGInputOutputElement>(); }
diff --git a/Engine/Plugins/Experimental/Web/WebAPI/Source/WebAPI/Public/WebAPIHttpRequest.h b/Engine/Plugins/Experimental/Web/WebAPI/Source/WebAPI/Public/WebAPIHttpRequest.h
--- a/Engine/Plugins/Experimental/Web/WebAPI/Source/WebAPI/Public/WebAPIHttpRequest.h
+++ b/Engine/Plugins/Experimental/Web/WebAPI/Source/WebAPI/Public/WebAPIHttpRequest.h
@@ -14,7 +14,7 @@ class TRequest : public TSharedFromThis<TRequest<PayloadType>>, public FNoncopya
 public:
 	void InitializeRequest(const FString& Route, const FString& InVerb = TEXT("POST"), const FString& InContentType = TEXT("json"));
 
-	UE_NODISCARD static TSharedRef<TRequest<PayloadType>> CopyHeadersToNewRequest(const TSharedRef<TRequest<PayloadType>>& Original);
+	[[nodiscard]] static TSharedRef<TRequest<PayloadType>> CopyHeadersToNewRequest(const TSharedRef<TRequest<PayloadType>>& Original);
 	
 	void Post() const;
 	void Post(PayloadType&& InPayload);
@@ -22,24 +22,24 @@ public:
 
 	void Get() const;
 
-	UE_NODISCARD TRequest<PayloadType>& AddHeader(TPair<FString, FString> NewHeader) const;
-	UE_NODISCARD TRequest<PayloadType>& BindCompletionCallback(TFunction<void(FHttpRequestPtr, FHttpResponsePtr, bool)> CompletionCallback);
+	[[nodiscard]] TRequest<PayloadType>& AddHeader(TPair<FString, FString> NewHeader) const;
+	[[nodiscard]] TRequest<PayloadType>& BindCompletionCallback(TFunction<void(FHttpRequestPtr, FHttpResponsePtr, bool)> CompletionCallback);
 	UE_DEPRECATED(5.3, "BindProgressCallback has been deprecated, use BindProgressCallback64 instead")
-	UE_NODISCARD TRequest<PayloadType>& BindProgressCallback(TFunction<void(FHttpRequestPtr, int32 /* BytesSent */, int32 /* BytesReceived */)> ProgressCallback);
-	UE_NODISCARD TRequest<PayloadType>& BindProgressCallback64(TFunction<void(FHttpRequestPtr, uint64 /* BytesSent */, uint64 /* BytesReceived */)> ProgressCallback);
-	UE_NODISCARD TRequest<PayloadType>& BindRetryCallback(TFunction<void(FHttpRequestPtr, FHttpResponsePtr, float /* SecondsToRetry */)> RetryCallback);
-	UE_NODISCARD TRequest<PayloadType>& BindHeaderReceivedCallback(TFunction<void(FHttpRequestPtr, const FString& /* HeaderName */, const FString& /* HeaderValue */)> HeaderReceivedCallback);
+	[[nodiscard]] TRequest<PayloadType>& BindProgressCallback(TFunction<void(FHttpRequestPtr, int32 /* BytesSent */, int32 /* BytesReceived */)> ProgressCallback);
+	[[nodiscard]] TRequest<PayloadType>& BindProgressCallback64(TFunction<void(FHttpRequestPtr, uint64 /* BytesSent */, uint64 /* BytesReceived */)> ProgressCallback);
+	[[nodiscard]] TRequest<PayloadType>& BindRetryCallback(TFunction<void(FHttpRequestPtr, FHttpResponsePtr, float /* SecondsToRetry */)> RetryCallback);
+	[[nodiscard]] TRequest<PayloadType>& BindHeaderReceivedCallback(TFunction<void(FHttpRequestPtr, const FString& /* HeaderName */, const FString& /* HeaderValue */)> HeaderReceivedCallback);
 	TRequest<PayloadType>& SetPayloadData(const PayloadType& InPayloadData);
 	TRequest<PayloadType>& SetPayloadData(TArray<uint8>& InPayloadData);
 
-	UE_NODISCARD bool HasPayload() const;
+	[[nodiscard]] bool HasPayload() const;
 
-	UE_NODISCARD FORCEINLINE FString GetVerb() const { return InternalRequest->GetVerb(); }
-	UE_NODISCARD FORCEINLINE FString GetContentType() const { return InternalRequest->GetContentType(); }
-	UE_NODISCARD FORCEINLINE FString GetRoute() const { return InternalRequest->GetURL(); }
+	[[nodiscard]] FORCEINLINE FString GetVerb() const { return InternalRequest->GetVerb(); }
+	[[nodiscard]] FORCEINLINE FString GetContentType() const { return InternalRequest->GetContentType(); }
+	[[nodiscard]] FORCEINLINE FString GetRoute() const { return InternalRequest->GetURL(); }
 
 	void SetHeaders(const TMap<FString, FString>& Headers) const;
-	UE_NODISCARD TMap<FString, FString> GetAllHeaders() const;
+	[[nodiscard]] TMap<FString, FString> GetAllHeaders() const;
 	
 	void SetRoute(const FString& InURL) const;
 	void SetVerb(const FString& InVerb) const;
diff --git a/Engine/Plugins/Interchange/Runtime/Source/Import/Public/MaterialX/MaterialXUtils/MaterialXManager.h b/Engine/Plugins/Interchange/Runtime/Source/Import/Public/MaterialX/MaterialXUtils/MaterialXManager.h
--- a/Engine/Plugins/Interchange/Runtime/Source/Import/Public/MaterialX/MaterialXUtils/MaterialXManager.h
+++ b/Engine/Plugins/Interchange/Runtime/Source/Import/Public/MaterialX/MaterialXUtils/MaterialXManager.h
@@ -60,7 +60,7 @@ private:
 			, NodeGroup{ std::forward<CategoryString>(NodeGroup) }
 		{}
 
-		UE_NODISCARD FORCEINLINE bool operator==(const FKeyCategoryNodegroup& Rhs) const
+		[[nodiscard]] FORCEINLINE bool operator==(const FKeyCategoryNodegroup& Rhs) const
 		{
 			return Category == Rhs.Category && NodeGroup == Rhs.NodeGroup;
 		}
diff --git a/Engine/Plugins/Online/OnlineServices/Source/OnlineServicesInterface/Public/Online/OnlineAsyncOpHandle.h b/Engine/Plugins/Online/OnlineServices/Source/OnlineServicesInterface/Public/Online/OnlineAsyncOpHandle.h
--- a/Engine/Plugins/Online/OnlineServices/Source/OnlineServicesInterface/Public/Online/OnlineAsyncOpHandle.h
+++ b/Engine/Plugins/Online/OnlineServices/Source/OnlineServicesInterface/Public/Online/OnlineAsyncOpHandle.h
@@ -156,7 +156,7 @@ namespace Private
 	}
 
 	template <typename DelegateSignature, typename ObjectOrCallableType, typename... VarTypes>
-	UE_NODISCARD inline TDelegate<DelegateSignature> ConstructDelegate(ObjectOrCallableType&& ObjectOrCallable, VarTypes&&... Vars)
+	[[nodiscard]] inline TDelegate<DelegateSignature> ConstructDelegate(ObjectOrCallableType&& ObjectOrCallable, VarTypes&&... Vars)
 	{
 		if constexpr((std::is_pointer_v<std::remove_reference_t<ObjectOrCallableType>> && !std::is_function_v<std::remove_reference_t<std::remove_pointer_t<ObjectOrCallableType>>>) || TSharedPtrHelper<std::remove_reference_t<ObjectOrCallableType>>::bIsSharedPtr)
 		{
@@ -342,7 +342,7 @@ public:
 	 *
 	 */
 	template<typename... ConstructDelegateParamTypes>
-	UE_NODISCARD FOnlineEventDelegateHandle Add(ConstructDelegateParamTypes&&... Params)
+	[[nodiscard]] FOnlineEventDelegateHandle Add(ConstructDelegateParamTypes&&... Params)
 	{
 		return Event.Add(Forward<ConstructDelegateParamTypes>(Params)...);
 	}
diff --git a/Engine/Plugins/Online/OnlineSubsystemEOS/Source/OnlineSubsystemEOS/Private/OnlineSubsystemEOSTypes.h b/Engine/Plugins/Online/OnlineSubsystemEOS/Source/OnlineSubsystemEOS/Private/OnlineSubsystemEOSTypes.h
--- a/Engine/Plugins/Online/OnlineSubsystemEOS/Source/OnlineSubsystemEOS/Private/OnlineSubsystemEOSTypes.h
+++ b/Engine/Plugins/Online/OnlineSubsystemEOS/Source/OnlineSubsystemEOS/Private/OnlineSubsystemEOSTypes.h
@@ -352,7 +352,7 @@ namespace OSSInternalCallback
 {
 	/** Create a callback for a non-SDK function that is tied to the lifetime of an arbitrary shared pointer. */
 	template <typename DelegateType, typename OwnerType, typename... CallbackArgs>
-	UE_NODISCARD DelegateType Create(const TSharedPtr<OwnerType, ESPMode::ThreadSafe>& InOwner,
+	[[nodiscard]] DelegateType Create(const TSharedPtr<OwnerType, ESPMode::ThreadSafe>& InOwner,
 		const TFunction<void(CallbackArgs...)>& InUserCallback)
 	{
 		const DelegateType& CheckOwnerThenExecute = DelegateType::CreateLambda(
diff --git a/Engine/Plugins/Runtime/GameplayAbilities/Source/GameplayAbilities/Public/AbilitySystemComponent.h b/Engine/Plugins/Runtime/GameplayAbilities/Source/GameplayAbilities/Public/AbilitySystemComponent.h
--- a/Engine/Plugins/Runtime/GameplayAbilities/Source/GameplayAbilities/Public/AbilitySystemComponent.h
+++ b/Engine/Plugins/Runtime/GameplayAbilities/Source/GameplayAbilities/Public/AbilitySystemComponent.h
@@ -577,7 +577,7 @@ class GAMEPLAYABILITIES_API UAbilitySystemComponent : public UGameplayTasksCompo
 		TagContainer.AppendTags(GetOwnedGameplayTags());
 	}
 
-	UE_NODISCARD FORCEINLINE const FGameplayTagContainer& GetOwnedGameplayTags() const
+	[[nodiscard]] FORCEINLINE const FGameplayTagContainer& GetOwnedGameplayTags() const
 	{
 		return GameplayTagCountContainer.GetExplicitGameplayTags();
 	}
@@ -624,7 +624,7 @@ class GAMEPLAYABILITIES_API UAbilitySystemComponent : public UGameplayTasksCompo
 		TagContainer.AppendTags(GetBlockedAbilityTags());
 	}
 
-	UE_NODISCARD FORCEINLINE const FGameplayTagContainer& GetBlockedAbilityTags() const
+	[[nodiscard]] FORCEINLINE const FGameplayTagContainer& GetBlockedAbilityTags() const
 	{
 		return BlockedAbilityTags.GetExplicitGameplayTags();
 	}
diff --git a/Engine/Plugins/Runtime/GameplayAbilities/Source/GameplayAbilities/Public/GameplayEffectTypes.h b/Engine/Plugins/Runtime/GameplayAbilities/Source/GameplayAbilities/Public/GameplayEffectTypes.h
--- a/Engine/Plugins/Runtime/GameplayAbilities/Source/GameplayAbilities/Public/GameplayEffectTypes.h
+++ b/Engine/Plugins/Runtime/GameplayAbilities/Source/GameplayAbilities/Public/GameplayEffectTypes.h
@@ -1394,7 +1394,7 @@ struct GAMEPLAYABILITIES_API FGameplayTagRequirements
 	bool operator!=(const FGameplayTagRequirements& Other) const;
 
 	/** Converts the RequireTags and IgnoreTags fields into an equivalent FGameplayTagQuery */
-	UE_NODISCARD FGameplayTagQuery ConvertTagFieldsToTagQuery() const;
+	[[nodiscard]] FGameplayTagQuery ConvertTagFieldsToTagQuery() const;
 };
 
 
diff --git a/Engine/Plugins/Runtime/MassGameplay/Source/MassSmartObjects/Public/MassSmartObjectHandler.h b/Engine/Plugins/Runtime/MassGameplay/Source/MassSmartObjects/Public/MassSmartObjectHandler.h
--- a/Engine/Plugins/Runtime/MassGameplay/Source/MassSmartObjects/Public/MassSmartObjectHandler.h
+++ b/Engine/Plugins/Runtime/MassGameplay/Source/MassSmartObjects/Public/MassSmartObjectHandler.h
@@ -46,7 +46,7 @@ struct MASSSMARTOBJECTS_API FMassSmartObjectHandler
 	 * @param Location The center of the query
 	 * @return Request identifier that can be used to try claiming a result once available
 	 */
-	UE_NODISCARD FMassSmartObjectRequestID FindCandidatesAsync(const FMassEntityHandle RequestingEntity, const FGameplayTagContainer& UserTags, const FGameplayTagQuery& ActivityRequirements, const FVector& Location) const;
+	[[nodiscard]] FMassSmartObjectRequestID FindCandidatesAsync(const FMassEntityHandle RequestingEntity, const FGameplayTagContainer& UserTags, const FGameplayTagQuery& ActivityRequirements, const FVector& Location) const;
 
 	/**
 	 * Creates an async request to build a list of compatible smart objects
@@ -56,7 +56,7 @@ struct MASSSMARTOBJECTS_API FMassSmartObjectHandler
 	 * @param LaneLocation The lane location as reference for the query
 	 * @return Request identifier that can be used to try claiming a result once available
 	 */
-	UE_NODISCARD FMassSmartObjectRequestID FindCandidatesAsync(const FMassEntityHandle RequestingEntity, const FGameplayTagContainer& UserTags, const FGameplayTagQuery& ActivityRequirements, const FZoneGraphCompactLaneLocation& LaneLocation) const;
+	[[nodiscard]] FMassSmartObjectRequestID FindCandidatesAsync(const FMassEntityHandle RequestingEntity, const FGameplayTagContainer& UserTags, const FGameplayTagQuery& ActivityRequirements, const FZoneGraphCompactLaneLocation& LaneLocation) const;
 
 	/**
 	 * Provides the result of a previously created request from FindCandidatesAsync to indicate if it has been processed
@@ -64,7 +64,7 @@ struct MASSSMARTOBJECTS_API FMassSmartObjectHandler
 	 * @param RequestID A valid request identifier (method will ensure otherwise)
 	 * @return The current request's result, nullptr if request not ready yet.
 	 */
-	UE_NODISCARD const FMassSmartObjectCandidateSlots* GetRequestCandidates(const FMassSmartObjectRequestID& RequestID) const;
+	[[nodiscard]] const FMassSmartObjectCandidateSlots* GetRequestCandidates(const FMassSmartObjectRequestID& RequestID) const;
 
 	/**
 	 * Deletes the request associated to the specified identifier
diff --git a/Engine/Plugins/Runtime/ModelViewViewModel/Source/ModelViewViewModel/Public/Bindings/MVVMBindingHelper.h b/Engine/Plugins/Runtime/ModelViewViewModel/Source/ModelViewViewModel/Public/Bindings/MVVMBindingHelper.h
--- a/Engine/Plugins/Runtime/ModelViewViewModel/Source/ModelViewViewModel/Public/Bindings/MVVMBindingHelper.h
+++ b/Engine/Plugins/Runtime/ModelViewViewModel/Source/ModelViewViewModel/Public/Bindings/MVVMBindingHelper.h
@@ -20,82 +20,82 @@ namespace UE::MVVM::BindingHelper
 {
 
 	/** Is the Property usable as a source (readable) by the binding system. It may required a Getter to read it. */
-	UE_NODISCARD MODELVIEWVIEWMODEL_API bool IsValidForSourceBinding(const FProperty* InProperty);
+	[[nodiscard]] MODELVIEWVIEWMODEL_API bool IsValidForSourceBinding(const FProperty* InProperty);
 
 	/**
 	 * Is the Function usable as a source (readable) by the binding system.
 	 * @note It may be a BlueprintGetter and binding to the Property would be better in the editor.
 	 */
-	UE_NODISCARD MODELVIEWVIEWMODEL_API bool IsValidForSourceBinding(const UFunction* InFunction);
+	[[nodiscard]] MODELVIEWVIEWMODEL_API bool IsValidForSourceBinding(const UFunction* InFunction);
 
 	/**
 	 * Is the Field usable as a source (readable) by the binding system.
 	 * @note It may be a BlueprintGetter and binding to the Property would be better in the editor.
 	 */
-	UE_NODISCARD MODELVIEWVIEWMODEL_API bool IsValidForSourceBinding(const FMVVMConstFieldVariant InField);
+	[[nodiscard]] MODELVIEWVIEWMODEL_API bool IsValidForSourceBinding(const FMVVMConstFieldVariant InField);
 
 	/** Is the Property usable as a destination (settable) by the binding system. It may required a Setter to write it. */
-	UE_NODISCARD MODELVIEWVIEWMODEL_API bool IsValidForDestinationBinding(const FProperty* InProperty);
+	[[nodiscard]] MODELVIEWVIEWMODEL_API bool IsValidForDestinationBinding(const FProperty* InProperty);
 
 	/**
 	 * Is the Function usable as a destination (settable) by the binding system.
 	 * @note It may be a BlueprintSetter and binding to the Property would be better in the editor.
 	 */
-	UE_NODISCARD MODELVIEWVIEWMODEL_API bool IsValidForDestinationBinding(const UFunction* InFunction);
+	[[nodiscard]] MODELVIEWVIEWMODEL_API bool IsValidForDestinationBinding(const UFunction* InFunction);
 
 	/**
 	 * Is the Field usable as a destination (settable) by the binding system.
 	 * @note It may be a BlueprintSetter and binding to the Property would be better in the editor.
 	 */
-	UE_NODISCARD MODELVIEWVIEWMODEL_API bool IsValidForDestinationBinding(const FMVVMConstFieldVariant InFunction);
+	[[nodiscard]] MODELVIEWVIEWMODEL_API bool IsValidForDestinationBinding(const FMVVMConstFieldVariant InFunction);
 
 	/**
 	 * Is the Function usable as a conversion function by the binding system without any wrapper.
 	 * A simple conversion takes a single argument and returns a single value.
 	 */
-	UE_NODISCARD MODELVIEWVIEWMODEL_API bool IsValidForSimpleRuntimeConversion(const UFunction* InFunction);
+	[[nodiscard]] MODELVIEWVIEWMODEL_API bool IsValidForSimpleRuntimeConversion(const UFunction* InFunction);
 
 	/**
 	 * Is the Function usable as a conversion function by the binding system with a wrapper.
 	 * A complex conversion takes a no argument and returns a single value.
 	 */
-	UE_NODISCARD MODELVIEWVIEWMODEL_API bool IsValidForComplexRuntimeConversion(const UFunction* InFunction);
+	[[nodiscard]] MODELVIEWVIEWMODEL_API bool IsValidForComplexRuntimeConversion(const UFunction* InFunction);
 
 #if WITH_EDITOR
 	/** Is the Property usable as a source by the binding system and can it be read directly or it requires a Getter. */
-	UE_NODISCARD MODELVIEWVIEWMODEL_API bool IsAccessibleDirectlyForSourceBinding(const FProperty* InProperty);
+	[[nodiscard]] MODELVIEWVIEWMODEL_API bool IsAccessibleDirectlyForSourceBinding(const FProperty* InProperty);
 
 	/** Is the Property usable as a destination by the binding system and can it be write directly or it requires a Setter. */
-	UE_NODISCARD MODELVIEWVIEWMODEL_API bool IsAccessibleDirectlyForDestinationBinding(const FProperty* InProperty);
+	[[nodiscard]] MODELVIEWVIEWMODEL_API bool IsAccessibleDirectlyForDestinationBinding(const FProperty* InProperty);
 
 	/** Is the Property usable as a source by the binding system and a Getter exists. */
-	UE_NODISCARD MODELVIEWVIEWMODEL_API bool IsAccessibleWithGetterForSourceBinding(const FProperty* InProperty);
+	[[nodiscard]] MODELVIEWVIEWMODEL_API bool IsAccessibleWithGetterForSourceBinding(const FProperty* InProperty);
 
 	/** Is the Property usable as a destination by the binding system and a Setter exists. */
-	UE_NODISCARD MODELVIEWVIEWMODEL_API bool IsAccessibleWithSetterForDestinationBinding(const FProperty* InProperty);
+	[[nodiscard]] MODELVIEWVIEWMODEL_API bool IsAccessibleWithSetterForDestinationBinding(const FProperty* InProperty);
 #endif
 
 	/**
 	 * Returns the Property or the Function that matches that BindingName.
 	 * @note It doesn't check for BlueprintGetter or BlueprintSetter.
 	 */
-	UE_NODISCARD MODELVIEWVIEWMODEL_API FMVVMFieldVariant FindFieldByName(const UStruct* Container, FMVVMBindingName BindingName);
+	[[nodiscard]] MODELVIEWVIEWMODEL_API FMVVMFieldVariant FindFieldByName(const UStruct* Container, FMVVMBindingName BindingName);
 
-	UE_NODISCARD MODELVIEWVIEWMODEL_API TValueOrError<const FProperty*, FText> TryGetPropertyTypeForSourceBinding(const FMVVMConstFieldVariant& InField);
-	UE_NODISCARD MODELVIEWVIEWMODEL_API TValueOrError<const FProperty*, FText> TryGetPropertyTypeForDestinationBinding(const FMVVMConstFieldVariant& InField);
-	UE_NODISCARD MODELVIEWVIEWMODEL_API TValueOrError<const FProperty*, FText> TryGetPropertyTypeForSourceBinding(const FProperty* InProperty);
-	UE_NODISCARD MODELVIEWVIEWMODEL_API TValueOrError<const FProperty*, FText> TryGetPropertyTypeForDestinationBinding(const FProperty* InProperty);
-	UE_NODISCARD MODELVIEWVIEWMODEL_API TValueOrError<const FProperty*, FText> TryGetPropertyTypeForSourceBinding(const UFunction* InFunction);
-	UE_NODISCARD MODELVIEWVIEWMODEL_API TValueOrError<const FProperty*, FText> TryGetPropertyTypeForDestinationBinding(const UFunction* InFunction);
+	[[nodiscard]] MODELVIEWVIEWMODEL_API TValueOrError<const FProperty*, FText> TryGetPropertyTypeForSourceBinding(const FMVVMConstFieldVariant& InField);
+	[[nodiscard]] MODELVIEWVIEWMODEL_API TValueOrError<const FProperty*, FText> TryGetPropertyTypeForDestinationBinding(const FMVVMConstFieldVariant& InField);
+	[[nodiscard]] MODELVIEWVIEWMODEL_API TValueOrError<const FProperty*, FText> TryGetPropertyTypeForSourceBinding(const FProperty* InProperty);
+	[[nodiscard]] MODELVIEWVIEWMODEL_API TValueOrError<const FProperty*, FText> TryGetPropertyTypeForDestinationBinding(const FProperty* InProperty);
+	[[nodiscard]] MODELVIEWVIEWMODEL_API TValueOrError<const FProperty*, FText> TryGetPropertyTypeForSourceBinding(const UFunction* InFunction);
+	[[nodiscard]] MODELVIEWVIEWMODEL_API TValueOrError<const FProperty*, FText> TryGetPropertyTypeForDestinationBinding(const UFunction* InFunction);
 
 	/** Returns the return property for a conversion function, or an error if the function is not valid as a conversion function. */
-	UE_NODISCARD MODELVIEWVIEWMODEL_API TValueOrError<const FProperty*, FText> TryGetReturnTypeForConversionFunction(const UFunction* InFunction);
+	[[nodiscard]] MODELVIEWVIEWMODEL_API TValueOrError<const FProperty*, FText> TryGetReturnTypeForConversionFunction(const UFunction* InFunction);
 
 	/** Returns the argument properties for a conversion function, or an error if the function is not valid as a conversion function. */
-	UE_NODISCARD MODELVIEWVIEWMODEL_API TValueOrError<TArray<const FProperty*>, FText> TryGetArgumentsForConversionFunction(const UFunction* InFunction);
+	[[nodiscard]] MODELVIEWVIEWMODEL_API TValueOrError<TArray<const FProperty*>, FText> TryGetArgumentsForConversionFunction(const UFunction* InFunction);
 
 	/** Are type the same or could be converted at runtime. */
-	UE_NODISCARD MODELVIEWVIEWMODEL_API bool ArePropertiesCompatible(const FProperty* Source, const FProperty* Destination);
+	[[nodiscard]] MODELVIEWVIEWMODEL_API bool ArePropertiesCompatible(const FProperty* Source, const FProperty* Destination);
 
 	/**
 	 * Returns the "return" property of the function. It can be a none-const ref/out argument.
@@ -103,7 +103,7 @@ namespace UE::MVVM::BindingHelper
 	 * void Foo(int&, double) -> returns int
 	 * void Foo(const int&, double) -> returns null
 	 */
-	UE_NODISCARD MODELVIEWVIEWMODEL_API const FProperty* GetReturnProperty(const UFunction* InFunction);
+	[[nodiscard]] MODELVIEWVIEWMODEL_API const FProperty* GetReturnProperty(const UFunction* InFunction);
 
 	/**
 	 * Returns the first argument property. Ref/Out argument have to be const.
@@ -111,18 +111,18 @@ namespace UE::MVVM::BindingHelper
 	 * void Foo(int& Out, double) -> returns double
 	 * void Foo(const int& Out, double) -> returns int
 	 */
-	UE_NODISCARD MODELVIEWVIEWMODEL_API const FProperty* GetFirstArgumentProperty(const UFunction* InFunction);
+	[[nodiscard]] MODELVIEWVIEWMODEL_API const FProperty* GetFirstArgumentProperty(const UFunction* InFunction);
 
 
 	/**
 	 * Returns all argument properties of this function.
 	 */
-	UE_NODISCARD MODELVIEWVIEWMODEL_API TArray<const FProperty*> GetAllArgumentProperties(const UFunction* InFunction);
+	[[nodiscard]] MODELVIEWVIEWMODEL_API TArray<const FProperty*> GetAllArgumentProperties(const UFunction* InFunction);
 
 	/**
 	 * Returns true if the given property is a valid function argument property.
 	 */
-	UE_NODISCARD MODELVIEWVIEWMODEL_API bool IsValidArgumentProperty(const FProperty* Property);
+	[[nodiscard]] MODELVIEWVIEWMODEL_API bool IsValidArgumentProperty(const FProperty* Property);
 
 	/**
 	 * Execute a binding that can be
diff --git a/Engine/Plugins/Runtime/ModelViewViewModel/Source/ModelViewViewModel/Public/Bindings/MVVMFieldPathHelper.h b/Engine/Plugins/Runtime/ModelViewViewModel/Source/ModelViewViewModel/Public/Bindings/MVVMFieldPathHelper.h
--- a/Engine/Plugins/Runtime/ModelViewViewModel/Source/ModelViewViewModel/Public/Bindings/MVVMFieldPathHelper.h
+++ b/Engine/Plugins/Runtime/ModelViewViewModel/Source/ModelViewViewModel/Public/Bindings/MVVMFieldPathHelper.h
@@ -40,26 +40,26 @@ namespace UE::MVVM::FieldPathHelper
 	};
 
 	/** */
-	UE_NODISCARD MODELVIEWVIEWMODEL_API TValueOrError<TArray<FMVVMConstFieldVariant>, FText> GenerateFieldPathList(const UClass* From, FStringView FieldPath, bool bForSourceBinding);
+	[[nodiscard]] MODELVIEWVIEWMODEL_API TValueOrError<TArray<FMVVMConstFieldVariant>, FText> GenerateFieldPathList(const UClass* From, FStringView FieldPath, bool bForSourceBinding);
 
 	/** */
-	UE_NODISCARD MODELVIEWVIEWMODEL_API TValueOrError<TArray<FMVVMConstFieldVariant>, FText> GenerateFieldPathList(TArrayView<const FMVVMConstFieldVariant> Fields, bool bForSourceBinding);
+	[[nodiscard]] MODELVIEWVIEWMODEL_API TValueOrError<TArray<FMVVMConstFieldVariant>, FText> GenerateFieldPathList(TArrayView<const FMVVMConstFieldVariant> Fields, bool bForSourceBinding);
 
 	/**
 	 * Generate the FParsedNotifyBindingInfo from the parsed field path.
 	 * This gives us which FieldId and Viewmodel that will be used by the path.
 	 * Note. The code may have changed since the binding was created in the editor. As an example, a sub-object can be a viewmodel now.
 	 */
-	UE_NODISCARD MODELVIEWVIEWMODEL_API TValueOrError<FParsedNotifyBindingInfo, FText> GetNotifyBindingInfoFromFieldPath(const UClass* Accessor, TArrayView<const FMVVMConstFieldVariant> Fields);
+	[[nodiscard]] MODELVIEWVIEWMODEL_API TValueOrError<FParsedNotifyBindingInfo, FText> GetNotifyBindingInfoFromFieldPath(const UClass* Accessor, TArrayView<const FMVVMConstFieldVariant> Fields);
 
 	/** */
-	UE_NODISCARD MODELVIEWVIEWMODEL_API FString ToString(TArrayView<const FMVVMFieldVariant> Fields);
-	UE_NODISCARD MODELVIEWVIEWMODEL_API FString ToString(TArrayView<const FMVVMConstFieldVariant> Fields);
-	UE_NODISCARD MODELVIEWVIEWMODEL_API FText ToText(TArrayView<const FMVVMFieldVariant> Fields);
-	UE_NODISCARD MODELVIEWVIEWMODEL_API FText ToText(TArrayView<const FMVVMConstFieldVariant> Fields);
+	[[nodiscard]] MODELVIEWVIEWMODEL_API FString ToString(TArrayView<const FMVVMFieldVariant> Fields);
+	[[nodiscard]] MODELVIEWVIEWMODEL_API FString ToString(TArrayView<const FMVVMConstFieldVariant> Fields);
+	[[nodiscard]] MODELVIEWVIEWMODEL_API FText ToText(TArrayView<const FMVVMFieldVariant> Fields);
+	[[nodiscard]] MODELVIEWVIEWMODEL_API FText ToText(TArrayView<const FMVVMConstFieldVariant> Fields);
 
 	/** Returns the object value from the FieldContext. */
-	UE_NODISCARD MODELVIEWVIEWMODEL_API TValueOrError<UObject*, void> EvaluateObjectProperty(const FFieldContext& InSource);
+	[[nodiscard]] MODELVIEWVIEWMODEL_API TValueOrError<UObject*, void> EvaluateObjectProperty(const FFieldContext& InSource);
 
 	/**
 	 * Returns the container that the field represents.
@@ -67,7 +67,7 @@ namespace UE::MVVM::FieldPathHelper
 	 * If the field is a Struct property, then returns the Struct.
 	 * If the field is a Function, then returns the return value property's container.
 	 */
-	UE_NODISCARD MODELVIEWVIEWMODEL_API TValueOrError<const UStruct*, void> GetFieldAsContainer(const UE::MVVM::FMVVMConstFieldVariant Field);
+	[[nodiscard]] MODELVIEWVIEWMODEL_API TValueOrError<const UStruct*, void> GetFieldAsContainer(const UE::MVVM::FMVVMConstFieldVariant Field);
 
 } // namespace
 
diff --git a/Engine/Plugins/Runtime/ModelViewViewModel/Source/ModelViewViewModel/Public/MVVMSubsystem.h b/Engine/Plugins/Runtime/ModelViewViewModel/Source/ModelViewViewModel/Public/MVVMSubsystem.h
--- a/Engine/Plugins/Runtime/ModelViewViewModel/Source/ModelViewViewModel/Public/MVVMSubsystem.h
+++ b/Engine/Plugins/Runtime/ModelViewViewModel/Source/ModelViewViewModel/Public/MVVMSubsystem.h
@@ -95,9 +95,9 @@ public:
 		FDirectionalBindingArgs BackwardArgs;
 	};
 
-	UE_NODISCARD TValueOrError<bool, FText> IsBindingValid(FConstDirectionalBindingArgs Args) const;
-	UE_NODISCARD TValueOrError<bool, FText> IsBindingValid(FDirectionalBindingArgs Args) const;
-	UE_NODISCARD TValueOrError<bool, FText> IsBindingValid(FBindingArgs Args) const;
+	[[nodiscard]] TValueOrError<bool, FText> IsBindingValid(FConstDirectionalBindingArgs Args) const;
+	[[nodiscard]] TValueOrError<bool, FText> IsBindingValid(FDirectionalBindingArgs Args) const;
+	[[nodiscard]] TValueOrError<bool, FText> IsBindingValid(FBindingArgs Args) const;
 };
 
 #if UE_ENABLE_INCLUDE_ORDER_DEPRECATED_IN_5_2
diff --git a/Engine/Plugins/Runtime/ModelViewViewModel/Source/ModelViewViewModel/Public/Types/MVVMBindingMode.h b/Engine/Plugins/Runtime/ModelViewViewModel/Source/ModelViewViewModel/Public/Types/MVVMBindingMode.h
--- a/Engine/Plugins/Runtime/ModelViewViewModel/Source/ModelViewViewModel/Public/Types/MVVMBindingMode.h
+++ b/Engine/Plugins/Runtime/ModelViewViewModel/Source/ModelViewViewModel/Public/Types/MVVMBindingMode.h
@@ -19,9 +19,9 @@ enum class EMVVMBindingMode : uint8
 
 namespace UE::MVVM
 {
-	UE_NODISCARD MODELVIEWVIEWMODEL_API bool IsForwardBinding(EMVVMBindingMode Mode);
-	UE_NODISCARD MODELVIEWVIEWMODEL_API bool IsBackwardBinding(EMVVMBindingMode Mode);
-	UE_NODISCARD MODELVIEWVIEWMODEL_API bool IsOneTimeBinding(EMVVMBindingMode Mode);
+	[[nodiscard]] MODELVIEWVIEWMODEL_API bool IsForwardBinding(EMVVMBindingMode Mode);
+	[[nodiscard]] MODELVIEWVIEWMODEL_API bool IsBackwardBinding(EMVVMBindingMode Mode);
+	[[nodiscard]] MODELVIEWVIEWMODEL_API bool IsOneTimeBinding(EMVVMBindingMode Mode);
 }
 
 #if UE_ENABLE_INCLUDE_ORDER_DEPRECATED_IN_5_2
diff --git a/Engine/Plugins/Runtime/ModelViewViewModel/Source/ModelViewViewModel/Public/Types/MVVMFieldVariant.h b/Engine/Plugins/Runtime/ModelViewViewModel/Source/ModelViewViewModel/Public/Types/MVVMFieldVariant.h
--- a/Engine/Plugins/Runtime/ModelViewViewModel/Source/ModelViewViewModel/Public/Types/MVVMFieldVariant.h
+++ b/Engine/Plugins/Runtime/ModelViewViewModel/Source/ModelViewViewModel/Public/Types/MVVMFieldVariant.h
@@ -61,12 +61,12 @@ namespace UE::MVVM
 			Type = EVariantType::Function;
 		}
 
-		UE_NODISCARD bool IsProperty() const
+		[[nodiscard]] bool IsProperty() const
 		{
 			return Type == EVariantType::Property;
 		}
 
-		UE_NODISCARD PropertyType* GetProperty() const
+		[[nodiscard]] PropertyType* GetProperty() const
 		{
 			check(Type == EVariantType::Property);
 			return Binding.Property;
@@ -78,12 +78,12 @@ namespace UE::MVVM
 			Type = EVariantType::Property;
 		}
 
-		UE_NODISCARD bool IsFunction() const
+		[[nodiscard]] bool IsFunction() const
 		{
 			return Type == EVariantType::Function;
 		}
 
-		UE_NODISCARD FunctionType* GetFunction() const
+		[[nodiscard]] FunctionType* GetFunction() const
 		{
 			check(Type == EVariantType::Function);
 			return Binding.Function;
@@ -95,7 +95,7 @@ namespace UE::MVVM
 			Type = EVariantType::Function;
 		}
 
-		UE_NODISCARD FName GetName() const
+		[[nodiscard]] FName GetName() const
 		{
 			if (Type == EVariantType::Property)
 			{
@@ -108,7 +108,7 @@ namespace UE::MVVM
 			return FName();
 		}
 
-		UE_NODISCARD UStruct* GetOwner() const
+		[[nodiscard]] UStruct* GetOwner() const
 		{
 			if (Type == EVariantType::Property)
 			{
@@ -121,7 +121,7 @@ namespace UE::MVVM
 			return nullptr;
 		}
 
-		UE_NODISCARD bool HasAnyFlags(EObjectFlags InFlags) const
+		[[nodiscard]] bool HasAnyFlags(EObjectFlags InFlags) const
 		{
 			if (Type == EVariantType::Property)
 			{
@@ -134,7 +134,7 @@ namespace UE::MVVM
 			return false;
 		}
 
-		UE_NODISCARD bool IsValid() const
+		[[nodiscard]] bool IsValid() const
 		{
 			if (Type == EVariantType::Property)
 			{
@@ -147,7 +147,7 @@ namespace UE::MVVM
 			return false;
 		}
 
-		UE_NODISCARD bool IsEmpty() const
+		[[nodiscard]] bool IsEmpty() const
 		{
 			return Type == EVariantType::Empty || Binding.Function == nullptr;
 		}
diff --git a/Engine/Plugins/Runtime/ModelViewViewModel/Source/ModelViewViewModel/Public/Types/MVVMFunctionContext.h b/Engine/Plugins/Runtime/ModelViewViewModel/Source/ModelViewViewModel/Public/Types/MVVMFunctionContext.h
--- a/Engine/Plugins/Runtime/ModelViewViewModel/Source/ModelViewViewModel/Public/Types/MVVMFunctionContext.h
+++ b/Engine/Plugins/Runtime/ModelViewViewModel/Source/ModelViewViewModel/Public/Types/MVVMFunctionContext.h
@@ -28,12 +28,12 @@ namespace UE::MVVM
 		}
 
 	public:
-		UE_NODISCARD ObjectType GetObject() const
+		[[nodiscard]] ObjectType GetObject() const
 		{
 			return ObjectPtr;
 		}
 
-		UE_NODISCARD FunctionType GetFunction() const
+		[[nodiscard]] FunctionType GetFunction() const
 		{
 			return FunctionPtr;
 		}
diff --git a/Engine/Plugins/Runtime/ModelViewViewModel/Source/ModelViewViewModel/Public/Types/MVVMObjectVariant.h b/Engine/Plugins/Runtime/ModelViewViewModel/Source/ModelViewViewModel/Public/Types/MVVMObjectVariant.h
--- a/Engine/Plugins/Runtime/ModelViewViewModel/Source/ModelViewViewModel/Public/Types/MVVMObjectVariant.h
+++ b/Engine/Plugins/Runtime/ModelViewViewModel/Source/ModelViewViewModel/Public/Types/MVVMObjectVariant.h
@@ -78,17 +78,17 @@ namespace UE::MVVM
 			Binding = VariantType(TInPlaceType<UScriptStructType>(), InScriptStruct, InValue);
 		}
 
-		UE_NODISCARD bool IsUObject() const
+		[[nodiscard]] bool IsUObject() const
 		{
 			return Binding.template IsType<UObjectType*>();
 		}
 
-		UE_NODISCARD UObjectType* GetUObject() const
+		[[nodiscard]] UObjectType* GetUObject() const
 		{
 			return Binding.template Get<UObjectType*>();
 		}
 
-		UE_NODISCARD const UClass* GetUClass() const
+		[[nodiscard]] const UClass* GetUClass() const
 		{
 			UObjectType* Instance = Binding.template Get<UObjectType*>();
 			return Instance ? Instance->GetClass() : nullptr;
@@ -99,17 +99,17 @@ namespace UE::MVVM
 			Binding.template Set<UObjectType*>(InValue);
 		}
 
-		UE_NODISCARD bool IsUScriptStruct() const
+		[[nodiscard]] bool IsUScriptStruct() const
 		{
 			return Binding.template IsType<UScriptStructType>();
 		}
 
-		UE_NODISCARD ContainerType GetUScriptStructData() const
+		[[nodiscard]] ContainerType GetUScriptStructData() const
 		{
 			return Binding.template Get<UScriptStructType>().Container;
 		}
 
-		UE_NODISCARD const UScriptStruct* GetUScriptStruct() const
+		[[nodiscard]] const UScriptStruct* GetUScriptStruct() const
 		{
 			return Binding.template Get<UScriptStructType>().ScriptStruct;
 		}
@@ -119,12 +119,12 @@ namespace UE::MVVM
 			Binding.template Set<UScriptStructType>(FScriptStructStackedContainer(InScriptStruct, InValue));
 		}
 
-		UE_NODISCARD ContainerType GetData() const
+		[[nodiscard]] ContainerType GetData() const
 		{
 			return IsUObject() ? reinterpret_cast<ContainerType>(GetUObject()) : IsUScriptStruct() ? GetUScriptStructData() : nullptr;
 		}
 
-		UE_NODISCARD const UStruct* GetOwner() const
+		[[nodiscard]] const UStruct* GetOwner() const
 		{
 			if (IsUObject())
 			{
@@ -137,12 +137,12 @@ namespace UE::MVVM
 			return nullptr;
 		}
 
-		UE_NODISCARD bool IsEmpty() const
+		[[nodiscard]] bool IsEmpty() const
 		{
 			return Binding.template IsType<FEmptyVariantState>();
 		}
 
-		UE_NODISCARD bool IsNull() const
+		[[nodiscard]] bool IsNull() const
 		{
 			return IsUObject() ? GetUObject() == nullptr : IsUScriptStruct() ? GetUScriptStructData() == nullptr : true;
 		}
diff --git a/Engine/Plugins/Runtime/ModelViewViewModel/Source/ModelViewViewModel/Public/View/MVVMViewClass.h b/Engine/Plugins/Runtime/ModelViewViewModel/Source/ModelViewViewModel/Public/View/MVVMViewClass.h
--- a/Engine/Plugins/Runtime/ModelViewViewModel/Source/ModelViewViewModel/Public/View/MVVMViewClass.h
+++ b/Engine/Plugins/Runtime/ModelViewViewModel/Source/ModelViewViewModel/Public/View/MVVMViewClass.h
@@ -485,25 +485,25 @@ public:
 
 public:
 	/** Should it automatically create the binding sources when the view is constructed. */
-	UE_NODISCARD bool DoesInitializeSourcesOnConstruct() const
+	[[nodiscard]] bool DoesInitializeSourcesOnConstruct() const
 	{
 		return bInitializeSourcesOnConstruct;
 	}
 
 	/** Should it automatically register and execute the bindings when the view is constructed. */
-	UE_NODISCARD bool DoesInitializeBindingsOnConstruct() const
+	[[nodiscard]] bool DoesInitializeBindingsOnConstruct() const
 	{
 		return bInitializeBindingsOnConstruct;
 	}
 	
 	/** Should it automatically register the events when the view is constructed. */
-	UE_NODISCARD bool DoesInitializeEventsOnConstruct() const
+	[[nodiscard]] bool DoesInitializeEventsOnConstruct() const
 	{
 		return bInitializeEventsOnConstruct;
 	}
 
 	/** Get the container of all the bindings. */
-	UE_NODISCARD const FMVVMCompiledBindingLibrary& GetBindingLibrary() const
+	[[nodiscard]] const FMVVMCompiledBindingLibrary& GetBindingLibrary() const
 	{
 		return BindingLibrary;
 	}
@@ -515,52 +515,52 @@ public:
 	}
 	
 	/** Get the list of all the needed viewmodel or widgets. */
-	UE_NODISCARD const TArrayView<const FMVVMViewClass_Source> GetSources() const
+	[[nodiscard]] const TArrayView<const FMVVMViewClass_Source> GetSources() const
 	{
 		return MakeArrayView(Sources);
 	}
 	
 	/** The shared source used by the view. */
-	UE_NODISCARD const FMVVMViewClass_Source& GetSource(FMVVMViewClass_SourceKey Key) const
+	[[nodiscard]] const FMVVMViewClass_Source& GetSource(FMVVMViewClass_SourceKey Key) const
 	{
 		check(Sources.IsValidIndex(Key.GetIndex()));
 		return Sources[Key.GetIndex()];
 	}
 
 	/** The list of bindings. A binding can be used by more than one source. */
-	UE_NODISCARD const TArrayView<const FMVVMViewClass_Binding> GetBindings() const
+	[[nodiscard]] const TArrayView<const FMVVMViewClass_Binding> GetBindings() const
 	{
 		return Bindings;
 	}
 
 	/** The binding (can be shared by more than one source). */
-	UE_NODISCARD const FMVVMViewClass_Binding& GetBinding(FMVVMViewClass_BindingKey Key) const
+	[[nodiscard]] const FMVVMViewClass_Binding& GetBinding(FMVVMViewClass_BindingKey Key) const
 	{
 		check(Bindings.IsValidIndex(Key.GetIndex()));
 		return Bindings[Key.GetIndex()];
 	}
 
 	/** The list of evaluate bindings. */
-	UE_NODISCARD const TArrayView<const FMVVMViewClass_EvaluateSource> GetEvaluateSources() const
+	[[nodiscard]] const TArrayView<const FMVVMViewClass_EvaluateSource> GetEvaluateSources() const
 	{
 		return EvaluateSources;
 	}
 
 	/** The evaluate binding. */
-	UE_NODISCARD const FMVVMViewClass_EvaluateSource& GetEvaluateSource(FMVVMViewClass_EvaluateBindingKey Key) const
+	[[nodiscard]] const FMVVMViewClass_EvaluateSource& GetEvaluateSource(FMVVMViewClass_EvaluateBindingKey Key) const
 	{
 		check(EvaluateSources.IsValidIndex(Key.GetIndex()));
 		return EvaluateSources[Key.GetIndex()];
 	}
 	
 	/** The list of events. */
-	UE_NODISCARD const TArrayView<const FMVVMViewClass_Event> GetEvents() const
+	[[nodiscard]] const TArrayView<const FMVVMViewClass_Event> GetEvents() const
 	{
 		return Events;
 	}
 
 	/** The event. */
-	UE_NODISCARD const FMVVMViewClass_Event& GetEvent(FMVVMViewClass_EventKey Key) const
+	[[nodiscard]] const FMVVMViewClass_Event& GetEvent(FMVVMViewClass_EventKey Key) const
 	{
 		check(Events.IsValidIndex(Key.GetIndex()));
 		return Events[Key.GetIndex()];
@@ -577,7 +577,7 @@ public:
 		MODELVIEWVIEWMODEL_API static FToStringArgs Short();
 		MODELVIEWVIEWMODEL_API static FToStringArgs All();
 	};
-	UE_NODISCARD FString ToString(FToStringArgs SourceArgs) const;
+	[[nodiscard]] FString ToString(FToStringArgs SourceArgs) const;
 #endif
 
 private:
diff --git a/Engine/Plugins/Runtime/ModelViewViewModel/Source/ModelViewViewModelBlueprint/Public/MVVMBlueprintViewModelContext.h b/Engine/Plugins/Runtime/ModelViewViewModel/Source/ModelViewViewModelBlueprint/Public/MVVMBlueprintViewModelContext.h
--- a/Engine/Plugins/Runtime/ModelViewViewModel/Source/ModelViewViewModelBlueprint/Public/MVVMBlueprintViewModelContext.h
+++ b/Engine/Plugins/Runtime/ModelViewViewModel/Source/ModelViewViewModelBlueprint/Public/MVVMBlueprintViewModelContext.h
@@ -31,7 +31,7 @@ enum class EMVVMBlueprintViewModelContextCreationType : uint8
 namespace UE::MVVM
 {
 #if WITH_EDITOR
-	UE_NODISCARD MODELVIEWVIEWMODELBLUEPRINT_API TArray<EMVVMBlueprintViewModelContextCreationType> GetAllowedContextCreationType(const UClass* Class);
+	[[nodiscard]] MODELVIEWVIEWMODELBLUEPRINT_API TArray<EMVVMBlueprintViewModelContextCreationType> GetAllowedContextCreationType(const UClass* Class);
 #endif
 }
 
diff --git a/Engine/Plugins/Runtime/WaveTable/Source/WaveTable/Public/WaveTableSettings.h b/Engine/Plugins/Runtime/WaveTable/Source/WaveTable/Public/WaveTableSettings.h
--- a/Engine/Plugins/Runtime/WaveTable/Source/WaveTable/Public/WaveTableSettings.h
+++ b/Engine/Plugins/Runtime/WaveTable/Source/WaveTable/Public/WaveTableSettings.h
@@ -14,7 +14,7 @@ struct FWaveTableBankEntry;
 namespace WaveTable
 {
 	// Converts a ratio [0.0, 1.0] to an index within the given number of array elements.
-	UE_NODISCARD static FORCEINLINE int32 RatioToIndex(float InRatio, int32 InArrayNum)
+	[[nodiscard]] static FORCEINLINE int32 RatioToIndex(float InRatio, int32 InArrayNum)
 	{
 		check(InArrayNum > 0);
 		check(InRatio >= 0.0f);
diff --git a/Engine/Plugins/Tests/CQTest/Source/CQTest/Public/Assert/NoDiscardAsserter.h b/Engine/Plugins/Tests/CQTest/Source/CQTest/Public/Assert/NoDiscardAsserter.h
--- a/Engine/Plugins/Tests/CQTest/Source/CQTest/Public/Assert/NoDiscardAsserter.h
+++ b/Engine/Plugins/Tests/CQTest/Source/CQTest/Public/Assert/NoDiscardAsserter.h
@@ -18,91 +18,91 @@ public:
 
 	void Fail(FString error);
 
-	UE_NODISCARD bool IsTrue(bool Condition);
-	UE_NODISCARD bool IsTrue(bool Condition, const char* FailureMessage);
-	UE_NODISCARD bool IsTrue(bool Condition, const TCHAR* FailureMessage);
-	UE_NODISCARD bool IsTrue(bool Condition, const FString& FailureMessage);
+	[[nodiscard]] bool IsTrue(bool Condition);
+	[[nodiscard]] bool IsTrue(bool Condition, const char* FailureMessage);
+	[[nodiscard]] bool IsTrue(bool Condition, const TCHAR* FailureMessage);
+	[[nodiscard]] bool IsTrue(bool Condition, const FString& FailureMessage);
 
-	UE_NODISCARD bool IsFalse(bool Condition);
-	UE_NODISCARD bool IsFalse(bool Condition, const char* FailureMessage);
-	UE_NODISCARD bool IsFalse(bool Condition, const TCHAR* FailureMessage);
-	UE_NODISCARD bool IsFalse(bool Condition, const FString& FailureMessage);
+	[[nodiscard]] bool IsFalse(bool Condition);
+	[[nodiscard]] bool IsFalse(bool Condition, const char* FailureMessage);
+	[[nodiscard]] bool IsFalse(bool Condition, const TCHAR* FailureMessage);
+	[[nodiscard]] bool IsFalse(bool Condition, const FString& FailureMessage);
 
 	template <typename T>
-	UE_NODISCARD bool IsNull(const T& Ptr);
+	[[nodiscard]] bool IsNull(const T& Ptr);
 
 	template <typename T>
-	UE_NODISCARD bool IsNull(const T& Ptr, const char* FailureMessage);
+	[[nodiscard]] bool IsNull(const T& Ptr, const char* FailureMessage);
 
 	template <typename T>
-	UE_NODISCARD bool IsNull(const T& Ptr, const TCHAR* FailureMessage);
+	[[nodiscard]] bool IsNull(const T& Ptr, const TCHAR* FailureMessage);
 
 	template <typename T>
-	UE_NODISCARD bool IsNull(const T& Ptr, const FString& FailureMessage);
+	[[nodiscard]] bool IsNull(const T& Ptr, const FString& FailureMessage);
 
 	template <typename T>
-	UE_NODISCARD bool IsNotNull(const T& Ptr);
+	[[nodiscard]] bool IsNotNull(const T& Ptr);
 
 	template <typename T>
-	UE_NODISCARD bool IsNotNull(const T& Ptr, const char* FailureMessage);
+	[[nodiscard]] bool IsNotNull(const T& Ptr, const char* FailureMessage);
 
 	template <typename T>
-	UE_NODISCARD bool IsNotNull(const T& Ptr, const TCHAR* FailureMessage);
+	[[nodiscard]] bool IsNotNull(const T& Ptr, const TCHAR* FailureMessage);
 
 	template <typename T>
-	UE_NODISCARD bool IsNotNull(const T& Ptr, const FString& FailureMessage);
+	[[nodiscard]] bool IsNotNull(const T& Ptr, const FString& FailureMessage);
 
 	template <typename TExpected, typename TActual>
-	UE_NODISCARD bool AreEqual(const TExpected& Expected, const TActual& Actual);
+	[[nodiscard]] bool AreEqual(const TExpected& Expected, const TActual& Actual);
 
 	template <typename TExpected, typename TActual>
-	UE_NODISCARD bool AreEqual(const TExpected& Expected, const TActual& Actual, const char* FailureMessage);
+	[[nodiscard]] bool AreEqual(const TExpected& Expected, const TActual& Actual, const char* FailureMessage);
 
 	template <typename TExpected, typename TActual>
-	UE_NODISCARD bool AreEqual(const TExpected& Expected, const TActual& Actual, const TCHAR* FailureMessage);
+	[[nodiscard]] bool AreEqual(const TExpected& Expected, const TActual& Actual, const TCHAR* FailureMessage);
 
 	template <typename TExpected, typename TActual>
-	UE_NODISCARD bool AreEqual(const TExpected& Expected, const TActual& Actual, const FString& FailureMessage);
+	[[nodiscard]] bool AreEqual(const TExpected& Expected, const TActual& Actual, const FString& FailureMessage);
 
-	UE_NODISCARD bool AreEqualIgnoreCase(const FString& Expected, const FString& Actual);
+	[[nodiscard]] bool AreEqualIgnoreCase(const FString& Expected, const FString& Actual);
 
-	UE_NODISCARD bool AreEqualIgnoreCase(const FString& Expected, const FString& Actual, const char* FailureMessage);
+	[[nodiscard]] bool AreEqualIgnoreCase(const FString& Expected, const FString& Actual, const char* FailureMessage);
 
-	UE_NODISCARD bool AreEqualIgnoreCase(const FString& Expected, const FString& Actual, const TCHAR* FailureMessage);
+	[[nodiscard]] bool AreEqualIgnoreCase(const FString& Expected, const FString& Actual, const TCHAR* FailureMessage);
 
-	UE_NODISCARD bool AreEqualIgnoreCase(const FString& Expected, const FString& Actual, const FString& FailureMessage);
+	[[nodiscard]] bool AreEqualIgnoreCase(const FString& Expected, const FString& Actual, const FString& FailureMessage);
 
 	template <typename TExpected, typename TActual>
-	UE_NODISCARD bool AreNotEqual(const TExpected& Expected, const TActual& Actual);
+	[[nodiscard]] bool AreNotEqual(const TExpected& Expected, const TActual& Actual);
 
 	template <typename TExpected, typename TActual>
-	UE_NODISCARD bool AreNotEqual(const TExpected& Expected, const TActual& Actual, const char* FailureMessage);
+	[[nodiscard]] bool AreNotEqual(const TExpected& Expected, const TActual& Actual, const char* FailureMessage);
 
 	template <typename TExpected, typename TActual>
-	UE_NODISCARD bool AreNotEqual(const TExpected& Expected, const TActual& Actual, const TCHAR* FailureMessage);
+	[[nodiscard]] bool AreNotEqual(const TExpected& Expected, const TActual& Actual, const TCHAR* FailureMessage);
 
 	template <typename TExpected, typename TActual>
-	UE_NODISCARD bool AreNotEqual(const TExpected& Expected, const TActual& Actual, const FString& FailureMessage);
+	[[nodiscard]] bool AreNotEqual(const TExpected& Expected, const TActual& Actual, const FString& FailureMessage);
 
-	UE_NODISCARD bool AreNotEqualIgnoreCase(const FString& Expected, const FString& Actual);
+	[[nodiscard]] bool AreNotEqualIgnoreCase(const FString& Expected, const FString& Actual);
 
-	UE_NODISCARD bool AreNotEqualIgnoreCase(const FString& Expected, const FString& Actual, const char* FailureMessage);
+	[[nodiscard]] bool AreNotEqualIgnoreCase(const FString& Expected, const FString& Actual, const char* FailureMessage);
 
-	UE_NODISCARD bool AreNotEqualIgnoreCase(const FString& Expected, const FString& Actual, const TCHAR* FailureMessage);
+	[[nodiscard]] bool AreNotEqualIgnoreCase(const FString& Expected, const FString& Actual, const TCHAR* FailureMessage);
 
-	UE_NODISCARD bool AreNotEqualIgnoreCase(const FString& Expected, const FString& Actual, const FString& FailureMessage);
+	[[nodiscard]] bool AreNotEqualIgnoreCase(const FString& Expected, const FString& Actual, const FString& FailureMessage);
 
 	template <typename TExpected, typename TActual, typename TEpsilon>
-	UE_NODISCARD bool IsNear(TExpected Expected, TActual Actual, TEpsilon Epsilon);
+	[[nodiscard]] bool IsNear(TExpected Expected, TActual Actual, TEpsilon Epsilon);
 
 	template <typename TExpected, typename TActual, typename TEpsilon>
-	UE_NODISCARD bool IsNear(TExpected Expected, TActual Actual, TEpsilon Epsilon, const char* FailureMessage);
+	[[nodiscard]] bool IsNear(TExpected Expected, TActual Actual, TEpsilon Epsilon, const char* FailureMessage);
 
 	template <typename TExpected, typename TActual, typename TEpsilon>
-	UE_NODISCARD bool IsNear(TExpected Expected, TActual Actual, TEpsilon Epsilon, const TCHAR* FailureMessage);
+	[[nodiscard]] bool IsNear(TExpected Expected, TActual Actual, TEpsilon Epsilon, const TCHAR* FailureMessage);
 
 	template <typename TExpected, typename TActual, typename TEpsilon>
-	UE_NODISCARD bool IsNear(TExpected Expected, TActual Actual, TEpsilon Epsilon, const FString& FailureMessage);
+	[[nodiscard]] bool IsNear(TExpected Expected, TActual Actual, TEpsilon Epsilon, const FString& FailureMessage);
 
 };
 
diff --git a/Engine/Plugins/Tests/CQTest/Source/CQTest/Public/Assert/NoDiscardAsserter.inl b/Engine/Plugins/Tests/CQTest/Source/CQTest/Public/Assert/NoDiscardAsserter.inl
--- a/Engine/Plugins/Tests/CQTest/Source/CQTest/Public/Assert/NoDiscardAsserter.inl
+++ b/Engine/Plugins/Tests/CQTest/Source/CQTest/Public/Assert/NoDiscardAsserter.inl
@@ -19,22 +19,22 @@ FORCEINLINE void FNoDiscardAsserter::Fail(FString Error)
 	TestRunner.AddError(Error);
 }
 
-UE_NODISCARD FORCEINLINE bool FNoDiscardAsserter::IsTrue(bool Condition)
+[[nodiscard]] FORCEINLINE bool FNoDiscardAsserter::IsTrue(bool Condition)
 {
 	return IsTrue(Condition, TEXT("Expected condition to be true."));
 }
 
-UE_NODISCARD FORCEINLINE bool FNoDiscardAsserter::IsTrue(bool Condition, const char* FailureMessage)
+[[nodiscard]] FORCEINLINE bool FNoDiscardAsserter::IsTrue(bool Condition, const char* FailureMessage)
 {
 	return IsTrue(Condition, FString(FailureMessage));
 }
 
-UE_NODISCARD FORCEINLINE bool FNoDiscardAsserter::IsTrue(bool Condition, const TCHAR* FailureMessage)
+[[nodiscard]] FORCEINLINE bool FNoDiscardAsserter::IsTrue(bool Condition, const TCHAR* FailureMessage)
 {
 	return IsTrue(Condition, FString(FailureMessage));
 }
 
-UE_NODISCARD FORCEINLINE bool FNoDiscardAsserter::IsTrue(bool Condition, const FString& FailureMessage)
+[[nodiscard]] FORCEINLINE bool FNoDiscardAsserter::IsTrue(bool Condition, const FString& FailureMessage)
 {
 	if (!Condition)
 	{
@@ -44,22 +44,22 @@ UE_NODISCARD FORCEINLINE bool FNoDiscardAsserter::IsTrue(bool Condition, const F
 	return true;
 }
 
-UE_NODISCARD FORCEINLINE bool FNoDiscardAsserter::IsFalse(bool Condition)
+[[nodiscard]] FORCEINLINE bool FNoDiscardAsserter::IsFalse(bool Condition)
 {
 	return IsFalse(Condition, TEXT("Expected condition to be false."));
 }
 
-UE_NODISCARD FORCEINLINE bool FNoDiscardAsserter::IsFalse(bool Condition, const char* FailureMessage)
+[[nodiscard]] FORCEINLINE bool FNoDiscardAsserter::IsFalse(bool Condition, const char* FailureMessage)
 {
 	return IsFalse(Condition, FString(FailureMessage));
 }
 
-UE_NODISCARD FORCEINLINE bool FNoDiscardAsserter::IsFalse(bool Condition, const TCHAR* FailureMessage)
+[[nodiscard]] FORCEINLINE bool FNoDiscardAsserter::IsFalse(bool Condition, const TCHAR* FailureMessage)
 {
 	return IsFalse(Condition, FString(FailureMessage));
 }
 
-UE_NODISCARD FORCEINLINE bool FNoDiscardAsserter::IsFalse(bool Condition, const FString& FailureMessage)
+[[nodiscard]] FORCEINLINE bool FNoDiscardAsserter::IsFalse(bool Condition, const FString& FailureMessage)
 {
 	if (Condition)
 	{
@@ -70,25 +70,25 @@ UE_NODISCARD FORCEINLINE bool FNoDiscardAsserter::IsFalse(bool Condition, const
 }
 
 template <typename T>
-UE_NODISCARD FORCEINLINE bool FNoDiscardAsserter::IsNull(const T& Ptr)
+[[nodiscard]] FORCEINLINE bool FNoDiscardAsserter::IsNull(const T& Ptr)
 {
 	return IsNull(Ptr, TEXT("Expected pointer to be Null"));
 }
 
 template <typename T>
-UE_NODISCARD FORCEINLINE bool FNoDiscardAsserter::IsNull(const T& Ptr, const char* FailureMessage)
+[[nodiscard]] FORCEINLINE bool FNoDiscardAsserter::IsNull(const T& Ptr, const char* FailureMessage)
 {
 	return IsNull(Ptr, FString(FailureMessage));
 }
 
 template <typename T>
-UE_NODISCARD FORCEINLINE bool FNoDiscardAsserter::IsNull(const T& Ptr, const TCHAR* FailureMessage)
+[[nodiscard]] FORCEINLINE bool FNoDiscardAsserter::IsNull(const T& Ptr, const TCHAR* FailureMessage)
 {
 	return IsNull(Ptr, FString(FailureMessage));
 }
 
 template <typename T>
-UE_NODISCARD FORCEINLINE bool FNoDiscardAsserter::IsNull(const T& Ptr, const FString& FailureMessage)
+[[nodiscard]] FORCEINLINE bool FNoDiscardAsserter::IsNull(const T& Ptr, const FString& FailureMessage)
 {
 	if (Ptr != nullptr)
 	{
@@ -99,25 +99,25 @@ UE_NODISCARD FORCEINLINE bool FNoDiscardAsserter::IsNull(const T& Ptr, const FSt
 }
 
 template <typename T>
-UE_NODISCARD FORCEINLINE bool FNoDiscardAsserter::IsNotNull(const T& Ptr)
+[[nodiscard]] FORCEINLINE bool FNoDiscardAsserter::IsNotNull(const T& Ptr)
 {
 	return IsNotNull(Ptr, TEXT("Expected pointer to be not Null"));
 }
 
 template <typename T>
-UE_NODISCARD FORCEINLINE bool FNoDiscardAsserter::IsNotNull(const T& Ptr, const char* FailureMessage)
+[[nodiscard]] FORCEINLINE bool FNoDiscardAsserter::IsNotNull(const T& Ptr, const char* FailureMessage)
 {
 	return IsNotNull(Ptr, FString(FailureMessage));
 }
 
 template <typename T>
-UE_NODISCARD FORCEINLINE bool FNoDiscardAsserter::IsNotNull(const T& Ptr, const TCHAR* FailureMessage)
+[[nodiscard]] FORCEINLINE bool FNoDiscardAsserter::IsNotNull(const T& Ptr, const TCHAR* FailureMessage)
 {
 	return IsNotNull(Ptr, FString(FailureMessage));
 }
 
 template <typename T>
-UE_NODISCARD FORCEINLINE bool FNoDiscardAsserter::IsNotNull(const T& Ptr, const FString& FailureMessage)
+[[nodiscard]] FORCEINLINE bool FNoDiscardAsserter::IsNotNull(const T& Ptr, const FString& FailureMessage)
 {
 	if (Ptr == nullptr)
 	{
@@ -128,7 +128,7 @@ UE_NODISCARD FORCEINLINE bool FNoDiscardAsserter::IsNotNull(const T& Ptr, const
 }
 
 template <typename TExpected, typename TActual>
-UE_NODISCARD FORCEINLINE bool FNoDiscardAsserter::AreEqual(const TExpected& Expected, const TActual& Actual)
+[[nodiscard]] FORCEINLINE bool FNoDiscardAsserter::AreEqual(const TExpected& Expected, const TActual& Actual)
 {
 	if (CQTestCondition::IsNotEqual(Expected, Actual))
 	{
@@ -139,19 +139,19 @@ UE_NODISCARD FORCEINLINE bool FNoDiscardAsserter::AreEqual(const TExpected& Expe
 }
 
 template <typename TExpected, typename TActual>
-UE_NODISCARD FORCEINLINE bool FNoDiscardAsserter::AreEqual(const TExpected& Expected, const TActual& Actual, const char* FailureMessage)
+[[nodiscard]] FORCEINLINE bool FNoDiscardAsserter::AreEqual(const TExpected& Expected, const TActual& Actual, const char* FailureMessage)
 {
 	return AreEqual(Expected, Actual, FString(FailureMessage));
 }
 
 template <typename TExpected, typename TActual>
-UE_NODISCARD FORCEINLINE bool FNoDiscardAsserter::AreEqual(const TExpected& Expected, const TActual& Actual, const TCHAR* FailureMessage)
+[[nodiscard]] FORCEINLINE bool FNoDiscardAsserter::AreEqual(const TExpected& Expected, const TActual& Actual, const TCHAR* FailureMessage)
 {
 	return AreEqual(Expected, Actual, FString(FailureMessage));
 }
 
 template <typename TExpected, typename TActual>
-UE_NODISCARD FORCEINLINE bool FNoDiscardAsserter::AreEqual(const TExpected& Expected, const TActual& Actual, const FString& FailureMessage)
+[[nodiscard]] FORCEINLINE bool FNoDiscardAsserter::AreEqual(const TExpected& Expected, const TActual& Actual, const FString& FailureMessage)
 {
 	if (CQTestCondition::IsNotEqual(Expected, Actual))
 	{
@@ -161,7 +161,7 @@ UE_NODISCARD FORCEINLINE bool FNoDiscardAsserter::AreEqual(const TExpected& Expe
 	return true;
 }
 
-UE_NODISCARD FORCEINLINE bool FNoDiscardAsserter::AreEqualIgnoreCase(const FString& Expected, const FString& Actual)
+[[nodiscard]] FORCEINLINE bool FNoDiscardAsserter::AreEqualIgnoreCase(const FString& Expected, const FString& Actual)
 {
 	if (CQTestCondition::IsNotEqualIgnoreCase(Expected, Actual))
 	{
@@ -171,17 +171,17 @@ UE_NODISCARD FORCEINLINE bool FNoDiscardAsserter::AreEqualIgnoreCase(const FStri
 	return true;
 }
 
-UE_NODISCARD FORCEINLINE bool FNoDiscardAsserter::AreEqualIgnoreCase(const FString& Expected, const FString& Actual, const char* FailureMessage)
+[[nodiscard]] FORCEINLINE bool FNoDiscardAsserter::AreEqualIgnoreCase(const FString& Expected, const FString& Actual, const char* FailureMessage)
 {
 	return AreEqualIgnoreCase(Expected, Actual, FString(FailureMessage));
 }
 
-UE_NODISCARD FORCEINLINE bool FNoDiscardAsserter::AreEqualIgnoreCase(const FString& Expected, const FString& Actual, const TCHAR* FailureMessage)
+[[nodiscard]] FORCEINLINE bool FNoDiscardAsserter::AreEqualIgnoreCase(const FString& Expected, const FString& Actual, const TCHAR* FailureMessage)
 {
 	return AreEqualIgnoreCase(Expected, Actual, FString(FailureMessage));
 }
 
-UE_NODISCARD FORCEINLINE bool FNoDiscardAsserter::AreEqualIgnoreCase(const FString& Expected, const FString& Actual, const FString& FailureMessage)
+[[nodiscard]] FORCEINLINE bool FNoDiscardAsserter::AreEqualIgnoreCase(const FString& Expected, const FString& Actual, const FString& FailureMessage)
 {
 	if (CQTestCondition::IsNotEqualIgnoreCase(Expected, Actual))
 	{
@@ -192,7 +192,7 @@ UE_NODISCARD FORCEINLINE bool FNoDiscardAsserter::AreEqualIgnoreCase(const FStri
 }
 
 template <typename TExpected, typename TActual>
-UE_NODISCARD FORCEINLINE bool FNoDiscardAsserter::AreNotEqual(const TExpected& Expected, const TActual& Actual)
+[[nodiscard]] FORCEINLINE bool FNoDiscardAsserter::AreNotEqual(const TExpected& Expected, const TActual& Actual)
 {
 	if (CQTestCondition::IsEqual(Expected, Actual))
 	{
@@ -203,19 +203,19 @@ UE_NODISCARD FORCEINLINE bool FNoDiscardAsserter::AreNotEqual(const TExpected& E
 }
 
 template <typename TExpected, typename TActual>
-UE_NODISCARD FORCEINLINE bool FNoDiscardAsserter::AreNotEqual(const TExpected& Expected, const TActual& Actual, const char* FailureMessage)
+[[nodiscard]] FORCEINLINE bool FNoDiscardAsserter::AreNotEqual(const TExpected& Expected, const TActual& Actual, const char* FailureMessage)
 {
 	return AreNotEqual(Expected, Actual, FString(FailureMessage));
 }
 
 template <typename TExpected, typename TActual>
-UE_NODISCARD FORCEINLINE bool FNoDiscardAsserter::AreNotEqual(const TExpected& Expected, const TActual& Actual, const TCHAR* FailureMessage)
+[[nodiscard]] FORCEINLINE bool FNoDiscardAsserter::AreNotEqual(const TExpected& Expected, const TActual& Actual, const TCHAR* FailureMessage)
 {
 	return AreNotEqual(Expected, Actual, FString(FailureMessage));
 }
 
 template <typename TExpected, typename TActual>
-UE_NODISCARD FORCEINLINE bool FNoDiscardAsserter::AreNotEqual(const TExpected& Expected, const TActual& Actual, const FString& FailureMessage)
+[[nodiscard]] FORCEINLINE bool FNoDiscardAsserter::AreNotEqual(const TExpected& Expected, const TActual& Actual, const FString& FailureMessage)
 {
 	if (CQTestCondition::IsEqual(Expected, Actual))
 	{
@@ -225,7 +225,7 @@ UE_NODISCARD FORCEINLINE bool FNoDiscardAsserter::AreNotEqual(const TExpected& E
 	return true;
 }
 
-UE_NODISCARD FORCEINLINE bool FNoDiscardAsserter::AreNotEqualIgnoreCase(const FString& Expected, const FString& Actual)
+[[nodiscard]] FORCEINLINE bool FNoDiscardAsserter::AreNotEqualIgnoreCase(const FString& Expected, const FString& Actual)
 {
 	if (CQTestCondition::IsEqualIgnoreCase(Expected, Actual))
 	{
@@ -235,17 +235,17 @@ UE_NODISCARD FORCEINLINE bool FNoDiscardAsserter::AreNotEqualIgnoreCase(const FS
 	return true;
 }
 
-UE_NODISCARD FORCEINLINE bool FNoDiscardAsserter::AreNotEqualIgnoreCase(const FString& Expected, const FString& Actual, const char* FailureMessage)
+[[nodiscard]] FORCEINLINE bool FNoDiscardAsserter::AreNotEqualIgnoreCase(const FString& Expected, const FString& Actual, const char* FailureMessage)
 {
 	return AreNotEqualIgnoreCase(Expected, Actual, FString(FailureMessage));
 }
 
-UE_NODISCARD FORCEINLINE bool FNoDiscardAsserter::AreNotEqualIgnoreCase(const FString& Expected, const FString& Actual, const TCHAR* FailureMessage)
+[[nodiscard]] FORCEINLINE bool FNoDiscardAsserter::AreNotEqualIgnoreCase(const FString& Expected, const FString& Actual, const TCHAR* FailureMessage)
 {
 	return AreNotEqualIgnoreCase(Expected, Actual, FString(FailureMessage));
 }
 
-UE_NODISCARD FORCEINLINE bool FNoDiscardAsserter::AreNotEqualIgnoreCase(const FString& Expected, const FString& Actual, const FString& FailureMessage)
+[[nodiscard]] FORCEINLINE bool FNoDiscardAsserter::AreNotEqualIgnoreCase(const FString& Expected, const FString& Actual, const FString& FailureMessage)
 {
 	if (CQTestCondition::IsEqualIgnoreCase(Expected, Actual))
 	{
@@ -256,7 +256,7 @@ UE_NODISCARD FORCEINLINE bool FNoDiscardAsserter::AreNotEqualIgnoreCase(const FS
 }
 
 template <typename TExpected, typename TActual, typename TEpsilon>
-UE_NODISCARD FORCEINLINE bool FNoDiscardAsserter::IsNear(TExpected Expected, TActual Actual, TEpsilon Epsilon)
+[[nodiscard]] FORCEINLINE bool FNoDiscardAsserter::IsNear(TExpected Expected, TActual Actual, TEpsilon Epsilon)
 {
 	if (CQTestCondition::IsNotNearlyEqual(Expected, Actual, Epsilon))
 	{
@@ -267,19 +267,19 @@ UE_NODISCARD FORCEINLINE bool FNoDiscardAsserter::IsNear(TExpected Expected, TAc
 }
 
 template <typename TExpected, typename TActual, typename TEpsilon>
-UE_NODISCARD FORCEINLINE bool FNoDiscardAsserter::IsNear(TExpected Expected, TActual Actual, TEpsilon Epsilon, const char* FailureMessage)
+[[nodiscard]] FORCEINLINE bool FNoDiscardAsserter::IsNear(TExpected Expected, TActual Actual, TEpsilon Epsilon, const char* FailureMessage)
 {
 	return IsNear(Expected, Actual, Epsilon, FString(FailureMessage));
 }
 
 template <typename TExpected, typename TActual, typename TEpsilon>
-UE_NODISCARD FORCEINLINE bool FNoDiscardAsserter::IsNear(TExpected Expected, TActual Actual, TEpsilon Epsilon, const TCHAR* FailureMessage)
+[[nodiscard]] FORCEINLINE bool FNoDiscardAsserter::IsNear(TExpected Expected, TActual Actual, TEpsilon Epsilon, const TCHAR* FailureMessage)
 {
 	return IsNear(Expected, Actual, Epsilon, FString(FailureMessage));
 }
 
 template <typename TExpected, typename TActual, typename TEpsilon>
-UE_NODISCARD FORCEINLINE bool FNoDiscardAsserter::IsNear(TExpected Expected, TActual Actual, TEpsilon Epsilon, const FString& FailureMessage)
+[[nodiscard]] FORCEINLINE bool FNoDiscardAsserter::IsNear(TExpected Expected, TActual Actual, TEpsilon Epsilon, const FString& FailureMessage)
 {
 	if (CQTestCondition::IsNotNearlyEqual(Expected, Actual, Epsilon))
 	{
diff --git a/Engine/Plugins/VirtualProduction/DMX/DMXEngine/Source/DMXEditor/Private/DMXZipper.h b/Engine/Plugins/VirtualProduction/DMX/DMXEngine/Source/DMXEditor/Private/DMXZipper.h
--- a/Engine/Plugins/VirtualProduction/DMX/DMXEngine/Source/DMXEditor/Private/DMXZipper.h
+++ b/Engine/Plugins/VirtualProduction/DMX/DMXEngine/Source/DMXEditor/Private/DMXZipper.h
@@ -42,13 +42,13 @@ public:
 	FDMXZipper();
 
 	/** Loads the specified zip file */
-	UE_NODISCARD bool LoadFromFile(const FString& Filename);
+	[[nodiscard]] bool LoadFromFile(const FString& Filename);
 
 	/** Loads the data as zip file */
-	UE_NODISCARD bool LoadFromData(const TArray64<uint8>& Data);
+	[[nodiscard]] bool LoadFromData(const TArray64<uint8>& Data);
 
 	/** Saves the zip to specified Filename */
-	UE_NODISCARD bool SaveToFile(const FString& Filename);
+	[[nodiscard]] bool SaveToFile(const FString& Filename);
 
 	/** Returns the File Names in the Zip. Note, may contain relative paths */
 	TArray<FString> GetFiles() const;
@@ -65,7 +65,7 @@ public:
 	void AddFile(const FString& RelativeFilePathAndName, const TArray64<uint8>& Data, const bool bCompress = false);
 
 	/** Gets the Data of a File in the Zip */
-	UE_NODISCARD bool GetFileContent(const FString& Filename, TArray64<uint8>& OutData);
+	[[nodiscard]] bool GetFileContent(const FString& Filename, TArray64<uint8>& OutData);
 
 	/** Helper to unzip a file within the zip as temp file. Deletes the file when running out of scope. */
 	struct FDMXScopedUnzipToTempFile
@@ -79,13 +79,13 @@ public:
 
 protected:
 	/** Writes a zipped File to the array writer */
-	UE_NODISCARD bool AddFileInternal(FArrayWriter& Writer, const FString& Path, const TArray64<uint8>& Data, const bool bCompress = true);
+	[[nodiscard]] bool AddFileInternal(FArrayWriter& Writer, const FString& Path, const TArray64<uint8>& Data, const bool bCompress = true);
 
 	/** Adds an End of Central Directory entry to the zip */
 	void AddEndOfCentralDirectory(FArrayWriter& Writer);
 
 	/** Parses the Zip File */
-	UE_NODISCARD bool Parse();
+	[[nodiscard]] bool Parse();
 
 	/** The Size of the Central Directory */
 	uint32 SizeOfCentralDirectory = 0;
diff --git a/Engine/Plugins/VirtualProduction/DMX/DMXEngine/Source/DMXEditor/Private/Exporters/DMXMVRExporter.h b/Engine/Plugins/VirtualProduction/DMX/DMXEngine/Source/DMXEditor/Private/Exporters/DMXMVRExporter.h
--- a/Engine/Plugins/VirtualProduction/DMX/DMXEngine/Source/DMXEditor/Private/Exporters/DMXMVRExporter.h
+++ b/Engine/Plugins/VirtualProduction/DMX/DMXEngine/Source/DMXEditor/Private/Exporters/DMXMVRExporter.h
@@ -26,10 +26,10 @@ private:
 	void ExportInternal(UDMXLibrary* DMXLibrary, const FString& FilePathAndName, FText& OutErrorReason);
 
 	/** Zips the GeneralSceneDescription.xml */
-	UE_NODISCARD bool ZipGeneralSceneDescription(const TSharedRef<FDMXZipper>& Zip, const UDMXMVRGeneralSceneDescription* GeneralSceneDescription, FText& OutErrorReason);
+	[[nodiscard]] bool ZipGeneralSceneDescription(const TSharedRef<FDMXZipper>& Zip, const UDMXMVRGeneralSceneDescription* GeneralSceneDescription, FText& OutErrorReason);
 
 	/** Zips the GDTFs from the Library */
-	UE_NODISCARD bool ZipGDTFs(const TSharedRef<FDMXZipper>& Zip, UDMXLibrary* DMXLibrary);
+	[[nodiscard]] bool ZipGDTFs(const TSharedRef<FDMXZipper>& Zip, UDMXLibrary* DMXLibrary);
 
 	/** Zips 3rd Party Data from the MVR Asset Import Data */
 	void ZipThirdPartyData(const TSharedRef<FDMXZipper>& Zip, const UDMXMVRGeneralSceneDescription* GeneralSceneDescription);
diff --git a/Engine/Plugins/VirtualProduction/DMX/DMXEngine/Source/DMXEditor/Private/Factories/DMXLibraryFromMVRImporter.cpp b/Engine/Plugins/VirtualProduction/DMX/DMXEngine/Source/DMXEditor/Private/Factories/DMXLibraryFromMVRImporter.cpp
--- a/Engine/Plugins/VirtualProduction/DMX/DMXEngine/Source/DMXEditor/Private/Factories/DMXLibraryFromMVRImporter.cpp
+++ b/Engine/Plugins/VirtualProduction/DMX/DMXEngine/Source/DMXEditor/Private/Factories/DMXLibraryFromMVRImporter.cpp
@@ -33,7 +33,7 @@
 namespace UE::DMX::DMXLibraryFromMVRImporter::Private
 {
 	/** Finds the file corresponding to the GDTF Spec in an MVR Zip file. Returns true if the file was found.  */
-	UE_NODISCARD bool FindGDTFFilenameInMVRZip(const TSharedRef<FDMXZipper>& MVRZip, const FString& GDTFSpec, FString& OutGDTFFilename)
+	[[nodiscard]] bool FindGDTFFilenameInMVRZip(const TSharedRef<FDMXZipper>& MVRZip, const FString& GDTFSpec, FString& OutGDTFFilename)
 	{
 		OutGDTFFilename.Reset();
 
diff --git a/Engine/Plugins/VirtualProduction/DMX/DMXEngine/Source/DMXEditor/Private/Factories/DMXLibraryFromMVRImporter.h b/Engine/Plugins/VirtualProduction/DMX/DMXEngine/Source/DMXEditor/Private/Factories/DMXLibraryFromMVRImporter.h
--- a/Engine/Plugins/VirtualProduction/DMX/DMXEngine/Source/DMXEditor/Private/Factories/DMXLibraryFromMVRImporter.h
+++ b/Engine/Plugins/VirtualProduction/DMX/DMXEngine/Source/DMXEditor/Private/Factories/DMXLibraryFromMVRImporter.h
@@ -30,7 +30,7 @@ class FDMXLibraryFromMVRImporter :
 {
 public:
 	/** Loads the MVR File */
-	UE_NODISCARD bool LoadMVRFile(const FString& InFilename);
+	[[nodiscard]] bool LoadMVRFile(const FString& InFilename);
 
 	/** Creates a modal window that allows the user to change import options. */
 	void UpdateImportOptionsFromModalWindow(UDMXLibraryFromMVRImportOptions* InOutImportOptions, bool& bOutCancelled);
diff --git a/Engine/Plugins/VirtualProduction/DMX/DMXEngine/Source/DMXEditor/Private/Widgets/FixturePatch/SDMXMVRFixtureListToolbar.h b/Engine/Plugins/VirtualProduction/DMX/DMXEngine/Source/DMXEditor/Private/Widgets/FixturePatch/SDMXMVRFixtureListToolbar.h
--- a/Engine/Plugins/VirtualProduction/DMX/DMXEngine/Source/DMXEditor/Private/Widgets/FixturePatch/SDMXMVRFixtureListToolbar.h
+++ b/Engine/Plugins/VirtualProduction/DMX/DMXEngine/Source/DMXEditor/Private/Widgets/FixturePatch/SDMXMVRFixtureListToolbar.h
@@ -31,7 +31,7 @@ public:
 	void Construct(const FArguments& InArgs, TWeakPtr<FDMXEditor> InDMXEditor);
 
 	/** Filters items according to the state of the search bar */
-	UE_NODISCARD TArray<TSharedPtr<FDMXMVRFixtureListItem>> FilterItems(const TArray<TSharedPtr<FDMXMVRFixtureListItem>>& Items);
+	[[nodiscard]] TArray<TSharedPtr<FDMXMVRFixtureListItem>> FilterItems(const TArray<TSharedPtr<FDMXMVRFixtureListItem>>& Items);
 
 private:
 	/** Generates the 'Add MVR Fixture' Dropdown Menu */
diff --git a/Engine/Plugins/VirtualProduction/DMX/DMXEngine/Source/DMXRuntime/Public/MVR/DMXMVRGeneralSceneDescription.h b/Engine/Plugins/VirtualProduction/DMX/DMXEngine/Source/DMXRuntime/Public/MVR/DMXMVRGeneralSceneDescription.h
--- a/Engine/Plugins/VirtualProduction/DMX/DMXEngine/Source/DMXRuntime/Public/MVR/DMXMVRGeneralSceneDescription.h
+++ b/Engine/Plugins/VirtualProduction/DMX/DMXEngine/Source/DMXRuntime/Public/MVR/DMXMVRGeneralSceneDescription.h
@@ -63,7 +63,7 @@ private:
 	void WriteFixturePatchToGeneralSceneDescription(const UDMXEntityFixturePatch& FixturePatch);
 
 	/** Parses a General Scene Description Xml File. Only ment to be used for initialization (ensured) */
-	UE_NODISCARD bool ParseGeneralSceneDescriptionXml(const TSharedRef<FXmlFile>& GeneralSceneDescriptionXml);
+	[[nodiscard]] bool ParseGeneralSceneDescriptionXml(const TSharedRef<FXmlFile>& GeneralSceneDescriptionXml);
 #endif
 
 	/** Returns the Fixture IDs currently in use, as number, sorted from lowest to highest */
diff --git a/Engine/Plugins/VirtualProduction/DMX/DMXEngine/Source/DMXRuntime/Public/MVR/Types/DMXMVRRootNode.h b/Engine/Plugins/VirtualProduction/DMX/DMXEngine/Source/DMXRuntime/Public/MVR/Types/DMXMVRRootNode.h
--- a/Engine/Plugins/VirtualProduction/DMX/DMXEngine/Source/DMXRuntime/Public/MVR/Types/DMXMVRRootNode.h
+++ b/Engine/Plugins/VirtualProduction/DMX/DMXEngine/Source/DMXRuntime/Public/MVR/Types/DMXMVRRootNode.h
@@ -36,7 +36,7 @@ public:
 	 * Does not merge, clears previous entries. 
 	 * Returns false if not a valid MVR General Scene Description Xml 
 	 */
-	UE_NODISCARD bool InitializeFromGeneralSceneDescriptionXml(const TSharedRef<FXmlFile>& GeneralSceneDescriptionXml);
+	[[nodiscard]] bool InitializeFromGeneralSceneDescriptionXml(const TSharedRef<FXmlFile>& GeneralSceneDescriptionXml);
 
 	/** Creates a DMX Xml Node as per MVR standard, returns nullptr if the Xml File cannot be created. */
 	TSharedPtr<FXmlFile> CreateXmlFile();
diff --git a/Engine/Plugins/VirtualProduction/LevelSnapshots/Source/LevelSnapshots/Private/Data/CustomSerialization/CustomObjectSerializationWrapper.h b/Engine/Plugins/VirtualProduction/LevelSnapshots/Source/LevelSnapshots/Private/Data/CustomSerialization/CustomObjectSerializationWrapper.h
--- a/Engine/Plugins/VirtualProduction/LevelSnapshots/Source/LevelSnapshots/Private/Data/CustomSerialization/CustomObjectSerializationWrapper.h
+++ b/Engine/Plugins/VirtualProduction/LevelSnapshots/Source/LevelSnapshots/Private/Data/CustomSerialization/CustomObjectSerializationWrapper.h
@@ -61,7 +61,7 @@ namespace UE::LevelSnapshots::Private
 
 
 	
-	UE_NODISCARD FRestoreObjectScope PreActorRestore_SnapshotWorld(
+	[[nodiscard]] FRestoreObjectScope PreActorRestore_SnapshotWorld(
 		AActor* EditorActor,
 		FCustomSerializationData& ActorSerializationData,
 		FWorldSnapshotData& WorldData,
@@ -70,7 +70,7 @@ namespace UE::LevelSnapshots::Private
 		UPackage* LocalisationSnapshotPackage
 		);
 
-	UE_NODISCARD FRestoreObjectScope PreActorRestore_EditorWorld(
+	[[nodiscard]] FRestoreObjectScope PreActorRestore_EditorWorld(
 		AActor* EditorActor,
 		FCustomSerializationData& ActorSerializationData,
 		FWorldSnapshotData& WorldData,
@@ -80,7 +80,7 @@ namespace UE::LevelSnapshots::Private
 		);
 	
 
-	UE_NODISCARD FRestoreObjectScope PreSubobjectRestore_SnapshotWorld(
+	[[nodiscard]] FRestoreObjectScope PreSubobjectRestore_SnapshotWorld(
 		UObject* Subobject,
 		const FSoftObjectPath& OriginalSubobjectPath,
 		FWorldSnapshotData& WorldData,
@@ -89,7 +89,7 @@ namespace UE::LevelSnapshots::Private
 		UPackage* LocalisationSnapshotPackage
 		);
 	
-	UE_NODISCARD FRestoreObjectScope PreSubobjectRestore_EditorWorld(
+	[[nodiscard]] FRestoreObjectScope PreSubobjectRestore_EditorWorld(
 		UObject* SnapshotObject,
 		UObject* EditorObject,
 		FWorldSnapshotData& WorldData,
diff --git a/Engine/Plugins/VirtualProduction/LevelSnapshots/Source/LevelSnapshots/Private/Util/SortedScopedLog.h b/Engine/Plugins/VirtualProduction/LevelSnapshots/Source/LevelSnapshots/Private/Util/SortedScopedLog.h
--- a/Engine/Plugins/VirtualProduction/LevelSnapshots/Source/LevelSnapshots/Private/Util/SortedScopedLog.h
+++ b/Engine/Plugins/VirtualProduction/LevelSnapshots/Source/LevelSnapshots/Private/Util/SortedScopedLog.h
@@ -75,7 +75,7 @@ public:
 
 	
 	/** Adds a log item. When it goes out of scope, its time is automatically added to LoggedItems. */
-	UE_NODISCARD FScopedLogItem AddScopedLogItem(FString Name, bool bItemLogCondition = true);
+	[[nodiscard]] FScopedLogItem AddScopedLogItem(FString Name, bool bItemLogCondition = true);
 
 	void AddCustomMeasuredLogItem(bool bLogCondition, FString Name, double Time) { if (bLogCondition) { AddCustomMeasuredLogItem(Name, Time); } }
 	/** Adds a time you measured yourself*/
diff --git a/Engine/Plugins/VirtualProduction/RemoteControl/Source/WebRemoteControl/Private/WebRemoteControlInternalUtils.h b/Engine/Plugins/VirtualProduction/RemoteControl/Source/WebRemoteControl/Private/WebRemoteControlInternalUtils.h
--- a/Engine/Plugins/VirtualProduction/RemoteControl/Source/WebRemoteControl/Private/WebRemoteControlInternalUtils.h
+++ b/Engine/Plugins/VirtualProduction/RemoteControl/Source/WebRemoteControl/Private/WebRemoteControlInternalUtils.h
@@ -96,7 +96,7 @@ namespace WebRemoteControlInternalUtils
 	 * @note InCompleteCallback will be called with an appropriate http response if the deserialization fails.
 	 */
 	template <typename RequestType>
-	UE_NODISCARD bool DeserializeRequestPayload(TConstArrayView<uint8> InTCHARPayload, const FHttpResultCallback* InCompleteCallback, RequestType& OutDeserializedRequest)
+	[[nodiscard]] bool DeserializeRequestPayload(TConstArrayView<uint8> InTCHARPayload, const FHttpResultCallback* InCompleteCallback, RequestType& OutDeserializedRequest)
 	{
 		FMemoryReaderView Reader(InTCHARPayload);
 		FJsonStructDeserializerBackend DeserializerBackend(Reader);
@@ -132,7 +132,7 @@ namespace WebRemoteControlInternalUtils
 	 * @param The wrapper structure to populate with the request's content.
 	 * @return Whether the deserialization was successful.
 	 */
-	UE_NODISCARD inline bool DeserializeWrappedRequestPayload(TConstArrayView<uint8> InTCHARPayload, const FHttpResultCallback* InCompleteCallback, FRCRequestWrapper& Wrapper)
+	[[nodiscard]] inline bool DeserializeWrappedRequestPayload(TConstArrayView<uint8> InTCHARPayload, const FHttpResultCallback* InCompleteCallback, FRCRequestWrapper& Wrapper)
 	{
 		if (!DeserializeRequestPayload(InTCHARPayload, InCompleteCallback, Wrapper))
 		{
@@ -155,7 +155,7 @@ namespace WebRemoteControlInternalUtils
 	 * @param OutErrorText If set, the string pointer will be populated with an error message on error.
 	 * @return Whether the delimiters were able to be found.
 	 */
-	UE_NODISCARD bool GetBatchRequestStructDelimiters(TConstArrayView<uint8> InTCHARPayload, TMap<int32, FBlockDelimiters>& OutStructParameters, FString* OutErrorText = nullptr);
+	[[nodiscard]] bool GetBatchRequestStructDelimiters(TConstArrayView<uint8> InTCHARPayload, TMap<int32, FBlockDelimiters>& OutStructParameters, FString* OutErrorText = nullptr);
 
 	/**
 	 * Get the struct delimiters for all the batched WebSocket requests.
@@ -164,7 +164,7 @@ namespace WebRemoteControlInternalUtils
 	 * @param OutErrorText If set, the string pointer will be populated with an error message on error.
 	 * @return Whether the delimiters were able to be found.
 	 */
-	UE_NODISCARD bool GetBatchWebSocketRequestStructDelimiters(TConstArrayView<uint8> InTCHARPayload, TArray<FBlockDelimiters>& OutStructParameters, FString* OutErrorText = nullptr);
+	[[nodiscard]] bool GetBatchWebSocketRequestStructDelimiters(TConstArrayView<uint8> InTCHARPayload, TArray<FBlockDelimiters>& OutStructParameters, FString* OutErrorText = nullptr);
 	
 	/**
 	 * Specialization of DeserializeRequestPayload that handles Batch requests.
@@ -177,7 +177,7 @@ namespace WebRemoteControlInternalUtils
 	 * @note InCompleteCallback will be called with an appropriate http response if the deserialization fails.
 	 */
 	template <>
-	UE_NODISCARD inline bool DeserializeRequestPayload(TConstArrayView<uint8> InTCHARPayload, const FHttpResultCallback* InCompleteCallback, FRCBatchRequest& OutDeserializedRequest)
+	[[nodiscard]] inline bool DeserializeRequestPayload(TConstArrayView<uint8> InTCHARPayload, const FHttpResultCallback* InCompleteCallback, FRCBatchRequest& OutDeserializedRequest)
 	{
 		FMemoryReaderView Reader(InTCHARPayload);
 		FJsonStructDeserializerBackend DeserializerBackend(Reader);
@@ -228,7 +228,7 @@ namespace WebRemoteControlInternalUtils
 	 * @note InCompleteCallback will be called with an appropriate http response if the deserialization fails.
 	 */
 	template <>
-	UE_NODISCARD inline bool DeserializeRequestPayload(TConstArrayView<uint8> InTCHARPayload, const FHttpResultCallback* InCompleteCallback, FRCWebSocketBatchRequest& OutDeserializedRequest)
+	[[nodiscard]] inline bool DeserializeRequestPayload(TConstArrayView<uint8> InTCHARPayload, const FHttpResultCallback* InCompleteCallback, FRCWebSocketBatchRequest& OutDeserializedRequest)
 	{
 		FMemoryReaderView Reader(InTCHARPayload);
 		FJsonStructDeserializerBackend DeserializerBackend(Reader);
@@ -298,7 +298,7 @@ namespace WebRemoteControlInternalUtils
 	 * @note InCompleteCallback will be called with an appropriate http response if the deserialization fails.
 	 */
 	template <typename RequestType>
-	UE_NODISCARD bool DeserializeRequest(const FHttpServerRequest& InRequest, const FHttpResultCallback* InCompleteCallback, RequestType& OutDeserializedRequest)
+	[[nodiscard]] bool DeserializeRequest(const FHttpServerRequest& InRequest, const FHttpResultCallback* InCompleteCallback, RequestType& OutDeserializedRequest)
 	{
 		static_assert(TIsDerivedFrom<RequestType, FRCRequest>::IsDerived, "Argument OutDeserializedRequest must derive from FRCRequest");
 		
@@ -325,14 +325,14 @@ namespace WebRemoteControlInternalUtils
 	 * 
 	 * @note InCompleteCallback will be called with an appropriate http response if the content type is not valid.
 	 */
-	UE_NODISCARD bool ValidateContentType(const FHttpServerRequest& InRequest, FString InContentType, const FHttpResultCallback& InCompleteCallback);
+	[[nodiscard]] bool ValidateContentType(const FHttpServerRequest& InRequest, FString InContentType, const FHttpResultCallback& InCompleteCallback);
 
 	/**
 	 * Check if a function call is valid since some objects//functions are disabled remotely for security reasons.
 	 * @param InRCCall The RC call to validate.
 	 * @param OutErrorText Optional error text.
 	 **/
-	UE_NODISCARD bool ValidateFunctionCall(const FRCCall& InRCCall, FString* OutErrorText);
+	[[nodiscard]] bool ValidateFunctionCall(const FRCCall& InRCCall, FString* OutErrorText);
 
 	/**
 	 * Add the desired content type to the http response headers.
diff --git a/Engine/Plugins/VirtualProduction/RemoteControl/Source/WebRemoteControl/Public/WebRemoteControlUtils.h b/Engine/Plugins/VirtualProduction/RemoteControl/Source/WebRemoteControl/Public/WebRemoteControlUtils.h
--- a/Engine/Plugins/VirtualProduction/RemoteControl/Source/WebRemoteControl/Public/WebRemoteControlUtils.h
+++ b/Engine/Plugins/VirtualProduction/RemoteControl/Source/WebRemoteControl/Public/WebRemoteControlUtils.h
@@ -66,7 +66,7 @@ namespace WebRemoteControlUtils
 	 * @return Whether the deserialization was successful.
 	 */
 	template <typename MessageType>
-	UE_NODISCARD bool DeserializeMessage(TConstArrayView<uint8> InTCHARPayload, MessageType& OutDeserializedMessage)
+	[[nodiscard]] bool DeserializeMessage(TConstArrayView<uint8> InTCHARPayload, MessageType& OutDeserializedMessage)
 	{
 		FMemoryReaderView Reader(InTCHARPayload);
 		TSharedRef<IStructDeserializerBackend> DeserializerBackend = CreateJsonDeserializerBackend(Reader);
diff --git a/Engine/Plugins/Web/HttpBlueprint/Source/HttpBlueprint/Public/HttpHeader.h b/Engine/Plugins/Web/HttpBlueprint/Source/HttpBlueprint/Public/HttpHeader.h
--- a/Engine/Plugins/Web/HttpBlueprint/Source/HttpBlueprint/Public/HttpHeader.h
+++ b/Engine/Plugins/Web/HttpBlueprint/Source/HttpBlueprint/Public/HttpHeader.h
@@ -18,12 +18,12 @@ struct HTTPBLUEPRINT_API FHttpHeader
 	 * Sets the internal headers.
 	 * Will overwrite any headers already in the map
 	 */
-	UE_NODISCARD FHttpHeader SetHeaders(const TMap<FString, FString>& NewHeaders);
+	[[nodiscard]] FHttpHeader SetHeaders(const TMap<FString, FString>& NewHeaders);
 
 	/**
 	 * Returns the header value associated with the Key
 	 */
-	UE_NODISCARD FString GetHeader(FString&& Key) const;
+	[[nodiscard]] FString GetHeader(FString&& Key) const;
 
 	/**
 	 * Adds a new header key/value pair to the internal header map
@@ -33,17 +33,17 @@ struct HTTPBLUEPRINT_API FHttpHeader
 	/**
 	 * Removes a header from the internal header map
 	 */
-	UE_NODISCARD bool RemoveHeader(FString&& HeaderToRemove);
+	[[nodiscard]] bool RemoveHeader(FString&& HeaderToRemove);
 
 	/**
 	 * Returns the internal header map as an array
 	 */
-	UE_NODISCARD TArray<FString> GetAllHeaders() const;
+	[[nodiscard]] TArray<FString> GetAllHeaders() const;
 
 	/**
 	 * Returns a copy of the internal header map
 	 */
-	UE_NODISCARD const TMap<FString, FString>& GetAllHeadersAsMap() const
+	[[nodiscard]] const TMap<FString, FString>& GetAllHeadersAsMap() const
 	{
 		return Headers;
 	}
@@ -51,7 +51,7 @@ struct HTTPBLUEPRINT_API FHttpHeader
 	/**
 	 * Returns if the header is valid or not. This is done by checking the map size
 	 */
-	UE_NODISCARD bool IsValid() const
+	[[nodiscard]] bool IsValid() const
 	{
 		return Headers.Num() > 0;
 	}
diff --git a/Engine/Plugins/Web/HttpBlueprint/Source/HttpBlueprintGraph/Private/K2Node_MakeRequestHeader.cpp b/Engine/Plugins/Web/HttpBlueprint/Source/HttpBlueprintGraph/Private/K2Node_MakeRequestHeader.cpp
--- a/Engine/Plugins/Web/HttpBlueprint/Source/HttpBlueprintGraph/Private/K2Node_MakeRequestHeader.cpp
+++ b/Engine/Plugins/Web/HttpBlueprint/Source/HttpBlueprintGraph/Private/K2Node_MakeRequestHeader.cpp
@@ -33,7 +33,7 @@ namespace UE::HttpBlueprint::Private
 
 	namespace Pin
 	{
-		UE_NODISCARD static FString MakePinName(int32 PinIndex)
+		[[nodiscard]] static FString MakePinName(int32 PinIndex)
 		{
 			const int32 PairIndex = PinIndex / 2;
 			if (PinIndex % 2 == 0)
@@ -44,7 +44,7 @@ namespace UE::HttpBlueprint::Private
 			return *FString::Printf(TEXT("%s %d"), *UE::HttpBlueprint::Private::MakeRequestHeaderGlobals::InputValuePrefix.ToString(), PairIndex);
 		}
 	
-		UE_NODISCARD static bool IsValidInputPin(const UEdGraphPin* InputPin)
+		[[nodiscard]] static bool IsValidInputPin(const UEdGraphPin* InputPin)
 		{
 			return InputPin
 				&& InputPin->Direction == EGPD_Input
diff --git a/Engine/Source/Editor/BlueprintGraph/Classes/EdGraphSchema_K2.h b/Engine/Source/Editor/BlueprintGraph/Classes/EdGraphSchema_K2.h
--- a/Engine/Source/Editor/BlueprintGraph/Classes/EdGraphSchema_K2.h
+++ b/Engine/Source/Editor/BlueprintGraph/Classes/EdGraphSchema_K2.h
@@ -1129,7 +1129,7 @@ public:
 		FName TargetFunction;
 		UClass* FunctionOwner = nullptr;
 	};
-	UE_NODISCARD virtual TOptional<FSearchForAutocastFunctionResults> SearchForAutocastFunction(const FEdGraphPinType& OutputPinType, const FEdGraphPinType& InputPinType) const;
+	[[nodiscard]] virtual TOptional<FSearchForAutocastFunctionResults> SearchForAutocastFunction(const FEdGraphPinType& OutputPinType, const FEdGraphPinType& InputPinType) const;
 
 	UE_DEPRECATED(5.2, "Use the FFindSpecializedConversionNodeResults variant.")
 	virtual bool FindSpecializedConversionNode(const UEdGraphPin* OutputPin, const UEdGraphPin* InputPin, bool bCreateNode, /*out*/ class UK2Node*& TargetNode) const;
@@ -1142,7 +1142,7 @@ public:
 	{
 		class UK2Node* TargetNode = nullptr;
 	};
-	UE_NODISCARD virtual TOptional<FFindSpecializedConversionNodeResults> FindSpecializedConversionNode(const FEdGraphPinType& OutputPinType, const UEdGraphPin& InputPin, bool bCreateNode) const;
+	[[nodiscard]] virtual TOptional<FFindSpecializedConversionNodeResults> FindSpecializedConversionNode(const FEdGraphPinType& OutputPinType, const UEdGraphPin& InputPin, bool bCreateNode) const;
 
 	/** Create menu for variable get/set nodes which refer to a variable which does not exist. */
 	void GetNonExistentVariableMenu(FToolMenuSection& Section, const UEdGraphNode* InGraphNode, UBlueprint* OwnerBlueprint) const;
diff --git a/Engine/Source/Editor/SequencerCore/Public/MVVM/Selection/SequencerCoreSelection.h b/Engine/Source/Editor/SequencerCore/Public/MVVM/Selection/SequencerCoreSelection.h
--- a/Engine/Source/Editor/SequencerCore/Public/MVVM/Selection/SequencerCoreSelection.h
+++ b/Engine/Source/Editor/SequencerCore/Public/MVVM/Selection/SequencerCoreSelection.h
@@ -52,14 +52,14 @@ public:
 	/**
 	 * Retrieve a scoped object that will suppress selection events for the duration of its lifetime on the stack
 	 */
-	UE_NODISCARD SEQUENCERCORE_API FSelectionEventSuppressor SuppressEvents();
+	[[nodiscard]] SEQUENCERCORE_API FSelectionEventSuppressor SuppressEvents();
 
 	/**
 	 * Retrieve a scoped object that will suppress selection events for the duration of its lifetime.
 	 * This function should only be used for longer running suppression such as those as part of a mouse-down operation.
 	 * The suppressor can be destroyed by resetting the unique pointer
 	 */
-	UE_NODISCARD SEQUENCERCORE_API TUniquePtr<FSelectionEventSuppressor> SuppressEventsLongRunning();
+	[[nodiscard]] SEQUENCERCORE_API TUniquePtr<FSelectionEventSuppressor> SuppressEventsLongRunning();
 
 	/**
 	 * Empty all the selection states within this selection class
diff --git a/Engine/Source/Runtime/Core/Internal/IO/IoStore.h b/Engine/Source/Runtime/Core/Internal/IO/IoStore.h
--- a/Engine/Source/Runtime/Core/Internal/IO/IoStore.h
+++ b/Engine/Source/Runtime/Core/Internal/IO/IoStore.h
@@ -213,9 +213,9 @@ struct FIoStoreTocResource
 	TArray<FIoStoreTocOnDemandChunkMeta> OnDemandChunkMeta;
 	TArray<FIoStoreTocOnDemandCompressedBlockMeta> OnDemandCompressedBlockMeta;
 
-	UE_NODISCARD CORE_API static FIoStatus Read(const TCHAR* TocFilePath, EIoStoreTocReadOptions ReadOptions, FIoStoreTocResource& OutTocResource);
+	[[nodiscard]] CORE_API static FIoStatus Read(const TCHAR* TocFilePath, EIoStoreTocReadOptions ReadOptions, FIoStoreTocResource& OutTocResource);
 
-	UE_NODISCARD CORE_API static TIoStatusOr<uint64> Write(const TCHAR* TocFilePath, FIoStoreTocResource& TocResource, uint32 CompressionBlockSize, uint64 MaxPartitionSize, const FIoContainerSettings& ContainerSettings);
+	[[nodiscard]] CORE_API static TIoStatusOr<uint64> Write(const TCHAR* TocFilePath, FIoStoreTocResource& TocResource, uint32 CompressionBlockSize, uint64 MaxPartitionSize, const FIoContainerSettings& ContainerSettings);
 
 	CORE_API static uint64 HashChunkIdWithSeed(int32 Seed, const FIoChunkId& ChunkId);
 };
diff --git a/Engine/Source/Runtime/Core/Private/Async/InheritedContext.cpp b/Engine/Source/Runtime/Core/Private/Async/InheritedContext.cpp
--- a/Engine/Source/Runtime/Core/Private/Async/InheritedContext.cpp
+++ b/Engine/Source/Runtime/Core/Private/Async/InheritedContext.cpp
@@ -5,7 +5,7 @@
 namespace UE
 {
 	// this method must be defined in a .cpp file to avoid exporting dependency on trace module
-	UE_NODISCARD FInheritedContextScope FInheritedContextBase::RestoreInheritedContext()
+	[[nodiscard]] FInheritedContextScope FInheritedContextBase::RestoreInheritedContext()
 	{
 		return FInheritedContextScope{
 #if ENABLE_LOW_LEVEL_MEM_TRACKER
