Fix for FMapTestSpawner to make sure the requested map exists
Added cleanup to stop active PIE sessions prior to creating a temporary level

#rnx

diff --git a/Engine/Plugins/Tests/CQTest/Source/CQTest/Private/Components/MapTestSpawner.cpp b/Engine/Plugins/Tests/CQTest/Source/CQTest/Private/Components/MapTestSpawner.cpp
--- a/Engine/Plugins/Tests/CQTest/Source/CQTest/Private/Components/MapTestSpawner.cpp
+++ b/Engine/Plugins/Tests/CQTest/Source/CQTest/Private/Components/MapTestSpawner.cpp
@@ -9,10 +9,12 @@
 #include "Engine/Engine.h"
 #include "Misc/Paths.h"
 #include "Editor.h"
+#include "Editor/UnrealEdEngine.h"
 #include "HAL/FileManager.h"
 #include "LevelEditor.h"
 #include "LevelEditorSubsystem.h"
 #include "Tests/AutomationEditorCommon.h"
+#include "UnrealEdGlobals.h"
 
 DEFINE_LOG_CATEGORY_STATIC(LogMapTest, Log, All);
 
@@ -53,6 +55,12 @@ FMapTestSpawner::FMapTestSpawner(const FString& MapDirectory, const FString& Map
 
 TUniquePtr<FMapTestSpawner> FMapTestSpawner::CreateFromTempLevel(FTestCommandBuilder& InCommandBuilder)
 {
+	if (IsValid(GUnrealEd->PlayWorld))
+	{
+		UE_LOG(LogMapTest, Verbose, TEXT("Active PIE session '%s' needs to be shutdown before a creation of a new level can occur."), *GUnrealEd->PlayWorld->GetMapName());
+		GUnrealEd->EndPlayMap();
+	}
+
 	FString MapName = GenerateUniqueMapName();
 	FString MapPath = FPaths::Combine(TempMapDirectory, MapName);
 	FString NewLevelPackage = FPackageName::FilenameToLongPackageName(MapPath);
@@ -74,9 +82,12 @@ void FMapTestSpawner::AddWaitUntilLoadedCommand(FAutomationTestBase* TestRunner)
 {
 	check(PieWorld == nullptr);
 
-	const FString FileName = FString::Printf(TEXT("%s.%s"), *MapName, *MapName);
-	const FString Path = FPaths::Combine(MapDirectory, FileName);
-	bool bOpened = AutomationOpenMap(Path);
+	FString PackagePath;
+	const FString Path = FPaths::Combine(MapDirectory, MapName);
+	bool bPackageExists = FPackageName::DoesPackageExist(Path, &PackagePath);
+	check(bPackageExists);
+
+	bool bOpened = AutomationOpenMap(PackagePath);
 	check(bOpened);
 
 	ADD_LATENT_AUTOMATION_COMMAND(FWaitUntil(*TestRunner, [&]() -> bool {
