Adding functionality for CQTest to handle input injection

diff --git a/Engine/Plugins/Tests/CQTest/CQTest.uplugin b/Engine/Plugins/Tests/CQTest/CQTest.uplugin
--- a/Engine/Plugins/Tests/CQTest/CQTest.uplugin
+++ b/Engine/Plugins/Tests/CQTest/CQTest.uplugin
@@ -24,5 +24,11 @@
 			"Type": "DeveloperTool",
 			"LoadingPhase": "Default"
 		}
+	],
+"Plugins": [
+		{
+			"Name": "EnhancedInput",
+			"Enabled": true
+		}
 	]
 }
\ No newline at end of file
diff --git a/Engine/Plugins/Tests/CQTest/Source/CQTest/CQTest.Build.cs b/Engine/Plugins/Tests/CQTest/Source/CQTest/CQTest.Build.cs
--- a/Engine/Plugins/Tests/CQTest/Source/CQTest/CQTest.Build.cs
+++ b/Engine/Plugins/Tests/CQTest/Source/CQTest/CQTest.Build.cs
@@ -15,7 +15,8 @@ public class CQTest : ModuleRules
 			new string[] {
 					"Core",
 					"CoreUObject",
-					"Engine"
+					"Engine",
+					"EnhancedInput"
 				 }
 			);
 
diff --git a/Engine/Plugins/Tests/CQTest/Source/CQTest/Private/Components/InputTestActions.cpp b/Engine/Plugins/Tests/CQTest/Source/CQTest/Private/Components/InputTestActions.cpp
new file mode 100644
--- /dev/null
+++ b/Engine/Plugins/Tests/CQTest/Source/CQTest/Private/Components/InputTestActions.cpp
@@ -0,0 +1,119 @@
+// Copyright Epic Games, Inc. All Rights Reserved.
+
+#include "Components/InputTestActions.h"
+
+#include "Engine/LocalPlayer.h"
+#include "EnhancedInputComponent.h"
+#include "EnhancedInputSubsystems.h"
+#include "GameFramework/Pawn.h"
+#include "GameFramework/PlayerController.h"
+
+namespace {
+
+const UInputAction* FindInputActionFromComponent(const UEnhancedInputComponent* InputComponent, const FString& InputActionName)
+{
+	for (const TUniquePtr<FEnhancedInputActionEventBinding>& Binding : InputComponent->GetActionEventBindings())
+	{
+		const UInputAction* Action = Binding->GetAction();
+		if (!IsValid(Action))
+		{
+			continue;
+		}
+		if (Action->GetName().Equals(InputActionName))
+		{
+			return Action;
+		}
+	}
+
+	return nullptr;
+}
+
+} //anonymous
+
+void FTestAction::operator()(const APawn* Pawn)
+{
+	check(IsValid(Pawn));
+
+	if (!IsValid(InputAction))
+	{
+		FindInputAction(Pawn);
+		check(InputAction);
+	}
+
+	APlayerController* PlayerController = CastChecked<APlayerController>(Pawn->GetController());
+	const ULocalPlayer* LocalPlayer = PlayerController->GetLocalPlayer();
+	check(IsValid(LocalPlayer));
+
+	UEnhancedInputLocalPlayerSubsystem* EnhancedInputLocalPlayerSubsystem = LocalPlayer->GetSubsystem<UEnhancedInputLocalPlayerSubsystem>();
+	check(IsValid(EnhancedInputLocalPlayerSubsystem));
+
+	UEnhancedPlayerInput* PlayerInput = EnhancedInputLocalPlayerSubsystem->GetPlayerInput();
+	check(IsValid(PlayerInput));
+	PlayerInput->InjectInputForAction(InputAction, InputActionValue);
+}
+
+void FTestAction::FindInputAction(const APawn* Pawn)
+{
+	if (const UEnhancedInputComponent* InputComponent = Cast<UEnhancedInputComponent>(Pawn->InputComponent))
+	{
+		InputAction = FindInputActionFromComponent(InputComponent, InputActionName);
+	}
+
+	if (!IsValid(InputAction))
+	{
+		APlayerController* PlayerController = CastChecked<APlayerController>(Pawn->GetController());
+		if (const UEnhancedInputComponent* InputComponent = Cast<UEnhancedInputComponent>(PlayerController->InputComponent))
+		{
+			InputAction = FindInputActionFromComponent(InputComponent, InputActionName);
+		}
+	}
+}
+
+FInputTestActions::~FInputTestActions()
+{
+	Reset();
+}
+
+void FInputTestActions::PerformAction(TFunction<void(const APawn* Pawn)> Action, TFunction<bool()> Predicate)
+{
+	if (Predicate)
+	{
+		if (!TickHandle.IsValid())
+		{
+			TickHandle = FTSTicker::GetCoreTicker().AddTicker(FTickerDelegate::CreateRaw(this, &FInputTestActions::Tick));
+		}
+
+		TestActions.Add(FTestActionPair(MoveTemp(Action), MoveTemp(Predicate)));
+	}
+	else
+	{
+		Action(Pawn);
+	}
+}
+
+void FInputTestActions::Reset()
+{
+	FTSTicker::GetCoreTicker().RemoveTicker(TickHandle);
+	TickHandle.Reset();
+}
+
+bool FInputTestActions::Tick(float DeltaTime)
+{
+	TestActions.RemoveAll([&](const FTestActionPair& TestActionItem) {
+		TFunctionRef<bool()> Predicate = TestActionItem.Get<1>();
+		return Predicate();
+	});
+
+	for (FTestActionPair& TestActionItem : TestActions)
+	{
+		TFunctionRef<void(const APawn* Pawn)> Action = TestActionItem.Get<0>();
+		Action(Pawn);
+	}
+
+	if (TestActions.IsEmpty())
+	{
+		Reset();
+	}
+
+	return true;
+}
\ No newline at end of file
diff --git a/Engine/Plugins/Tests/CQTest/Source/CQTest/Public/Components/InputTestActions.h b/Engine/Plugins/Tests/CQTest/Source/CQTest/Public/Components/InputTestActions.h
new file mode 100644
--- /dev/null
+++ b/Engine/Plugins/Tests/CQTest/Source/CQTest/Public/Components/InputTestActions.h
@@ -0,0 +1,77 @@
+// Copyright Epic Games, Inc. All Rights Reserved.
+
+#pragma once
+
+#include "Containers/Ticker.h"
+#include "InputAction.h"
+#include "InputActionValue.h"
+
+class APawn;
+
+// Class for testing input of a Pawn by injecting InputActions
+class CQTEST_API FTestAction
+{
+public:
+	virtual ~FTestAction()
+	{
+	}
+
+	/**
+	 * Custom input functionality to be applied on the provided Pawn.
+	 *
+	 * @param Pawn - Actor which the input logic will be applied to.
+	 */
+	virtual void operator()(const APawn* Pawn);
+
+	FString InputActionName;
+	FInputActionValue InputActionValue;
+
+private:
+	/**
+	 * Finds the appropriate InputAction mapping from the Pawn using the name provided from InputActionName.
+	 *
+	 * @param Pawn - Pawn to search for the specified InputAction mapping.
+	 */
+	void FindInputAction(const APawn* Pawn);
+
+	const UInputAction* InputAction{ nullptr };
+};
+
+// Class for processing FTestAction objects
+class CQTEST_API FInputTestActions
+{
+public:
+	FInputTestActions(APawn* InPawn) : Pawn(InPawn)
+	{
+	}
+
+	virtual ~FInputTestActions();
+
+protected:
+	/**
+	 * Processes the action within the current tick.
+	 *
+	 * @param Action - Function with the logic to be processed on the given Pawn.
+	 * @param Predicate - Function used to determine if the Action should be executed.
+	 */
+	void PerformAction(TFunction<void(const APawn* Pawn)> Action, TFunction<bool()> Predicate = nullptr);
+
+	/** Clears all active timers. */
+	virtual void Reset();
+
+	APawn* Pawn{ nullptr };
+
+private:
+	/**
+	 * Processes repeat actions every tick.
+	 *
+	 * @param DeltaTime - Time elapsed since the last tick.
+	 * @return true if tick has completed.
+	 */
+	bool Tick(float DeltaTime);
+
+	FTSTicker::FDelegateHandle TickHandle;
+
+	using FTestActionPair = TPair<TFunction<void(const APawn* Pawn)>, TFunction<bool()>>;
+	TArray<FTestActionPair> TestActions;
+};
\ No newline at end of file
diff --git a/Engine/Plugins/Tests/CQTest/Source/CQTestTests/CQTestTests.Build.cs b/Engine/Plugins/Tests/CQTest/Source/CQTestTests/CQTestTests.Build.cs
--- a/Engine/Plugins/Tests/CQTest/Source/CQTestTests/CQTestTests.Build.cs
+++ b/Engine/Plugins/Tests/CQTest/Source/CQTestTests/CQTestTests.Build.cs
@@ -14,6 +14,8 @@ public class CQTestTests : ModuleRules
 					"Core",
 					"CoreUObject",
 					"Engine",
+					"InputCore",
+					"EnhancedInput",
 					"CQTest"
 				 }
 			);
diff --git a/Engine/Plugins/Tests/CQTest/Source/CQTestTests/Private/CQTestInputTestHelper.cpp b/Engine/Plugins/Tests/CQTest/Source/CQTestTests/Private/CQTestInputTestHelper.cpp
new file mode 100644
--- /dev/null
+++ b/Engine/Plugins/Tests/CQTest/Source/CQTestTests/Private/CQTestInputTestHelper.cpp
@@ -0,0 +1,142 @@
+// Copyright Epic Games, Inc. All Rights Reserved.
+
+#include "CQTestInputTestHelper.h"
+
+#include "Engine/World.h"
+#include "InputCoreTypes.h"
+#include "InputMappingContext.h"
+
+void UTestEnhancedInputSubsystem::Init(APlayerController* InPlayerController)
+{
+	check(InPlayerController);
+
+	PlayerController = InPlayerController;
+	InputComponent = CastChecked<UEnhancedInputComponent>(PlayerController->InputComponent);
+	PlayerInput = CastChecked<UEnhancedPlayerInput>(PlayerController->PlayerInput);
+
+	InitalizeUserSettings();
+	SetupTestBindings();
+}
+
+void UTestEnhancedInputSubsystem::MappingListener(const FInputActionInstance& Instance)
+{
+	const UInputAction* SourceAction = Instance.GetSourceAction();
+	check(SourceAction);
+
+	ETriggerEvent* TriggerEvent = InputActionTriggerEventMap.Find(SourceAction->GetName());
+	if (TriggerEvent)
+	{
+		*TriggerEvent = Instance.GetTriggerEvent();
+	}
+	else
+	{
+		InputActionTriggerEventMap.Add(SourceAction->GetName(), Instance.GetTriggerEvent());
+	}
+}
+
+ETriggerEvent* UTestEnhancedInputSubsystem::GetEventForAction(const FString& InputActionName)
+{
+	return InputActionTriggerEventMap.Find(InputActionName);
+}
+
+void UTestEnhancedInputSubsystem::SetupTestBindings()
+{
+	InputMappingContext = NewObject<UInputMappingContext>(PlayerController, "TestContext");
+	AddMappingContext(InputMappingContext, 0);
+
+	BindInputAction(FCQTestInputSubsystemHelper::TestButtonActionName, EInputActionValueType::Boolean, { EKeys::A });
+	BindInputAction(FCQTestInputSubsystemHelper::TestAxisActionName, EInputActionValueType::Axis1D, { EKeys::Gamepad_LeftTriggerAxis, EKeys::MouseX });
+
+	// Generate a live mapping on the player
+	FModifyContextOptions Options;
+	Options.bForceImmediately = true;
+	RequestRebuildControlMappings(Options);
+}
+
+void UTestEnhancedInputSubsystem::BindInputAction(const FString& InputActionName, EInputActionValueType InputActionValueType, const TArray<FKey>& Keys)
+{
+	UInputAction* Action = NewObject<UInputAction>(PlayerController, *InputActionName);
+	Action->ValueType = InputActionValueType;
+
+	// Bind action to the binding targets so we can check if they were called correctly, but only the first time we bind the action!
+	InputComponent->BindAction(Action, ETriggerEvent::Started, this, &UTestEnhancedInputSubsystem::MappingListener);
+	InputComponent->BindAction(Action, ETriggerEvent::Ongoing, this, &UTestEnhancedInputSubsystem::MappingListener);
+	InputComponent->BindAction(Action, ETriggerEvent::Canceled, this, &UTestEnhancedInputSubsystem::MappingListener);
+	InputComponent->BindAction(Action, ETriggerEvent::Completed, this, &UTestEnhancedInputSubsystem::MappingListener);
+	InputComponent->BindAction(Action, ETriggerEvent::Triggered, this, &UTestEnhancedInputSubsystem::MappingListener);
+
+	// Initialise input action mapping in context to be used for testing
+	for (const FKey& Key : Keys)
+	{
+		InputMappingContext->MapKey(Action, Key);
+	}
+}
+
+bool FCQTestInputSubsystemHelper::ActionExpectedEvent(const FString& InputActionName, ETriggerEvent ExpectedTriggerEvent)
+{
+	check(InputSubsystem);
+
+	ETriggerEvent* TriggerEvent = InputSubsystem->GetEventForAction(InputActionName);
+	if (!TriggerEvent)
+	{
+		return false;
+	}
+	
+	return (*TriggerEvent) == ExpectedTriggerEvent;
+}
+
+void FCQTestInputSubsystemHelper::InitializePlayerControllerInput()
+{
+	check(Pawn);
+
+	if (!Pawn->IsPlayerControlled())
+	{
+		UWorld* CurrentWorld = Pawn->GetWorld();
+		check(CurrentWorld);
+
+		PlayerController = NewObject<APlayerController>(CurrentWorld);
+		PlayerController->Possess(Pawn);
+		Pawn->PossessedBy(PlayerController);
+	}
+	else
+	{
+		PlayerController = CastChecked<APlayerController>(Pawn->GetController());
+		check(PlayerController);
+	}
+
+	PlayerController->InputComponent = NewObject<UEnhancedInputComponent>(PlayerController);
+	PlayerController->PlayerInput = NewObject<UEnhancedPlayerInput>(PlayerController);
+	PlayerController->InitInputSystem();
+
+	// Keep a strong reference to the newly created input system during the test execution
+	InputSubsystem = TStrongObjectPtr<UTestEnhancedInputSubsystem>(NewObject<UTestEnhancedInputSubsystem>());
+	InputSubsystem->Init(PlayerController);
+}
+
+void FCQTestPawnTestActions::PressButton(const FString& ButtonActionName)
+{
+	PerformAction(FPressButtonAction{ ButtonActionName });
+}
+
+void FCQTestPawnTestActions::HoldAxis(const FString& ExampleAxisActionName, const FInputActionValue& ActionValue, FTimespan&& Duration)
+{
+	PerformAction(FHoldAxisAction{ ExampleAxisActionName, ActionValue }, [this, Duration = MoveTemp(Duration)]() -> bool {
+		if (StartTime.GetTicks() == 0)
+		{
+			StartTime = FDateTime::UtcNow();
+		}
+
+		FTimespan Elapsed = FDateTime::UtcNow() - StartTime;
+		return Elapsed >= Duration;
+		});
+}
+
+bool FCQTestPawnTestActions::IsTriggered(const FString& InputActionName)
+{
+	return InputSubsystemHelper->ActionExpectedEvent(InputActionName, ETriggerEvent::Triggered);
+}
+
+bool FCQTestPawnTestActions::IsCompleted(const FString& InputActionName)
+{
+	return InputSubsystemHelper->ActionExpectedEvent(InputActionName, ETriggerEvent::Completed);
+}
\ No newline at end of file
diff --git a/Engine/Plugins/Tests/CQTest/Source/CQTestTests/Private/CQTestInputTestHelper.h b/Engine/Plugins/Tests/CQTest/Source/CQTestTests/Private/CQTestInputTestHelper.h
new file mode 100644
--- /dev/null
+++ b/Engine/Plugins/Tests/CQTest/Source/CQTestTests/Private/CQTestInputTestHelper.h
@@ -0,0 +1,114 @@
+// Copyright Epic Games, Inc. All Rights Reserved.
+
+#pragma once
+
+#include "Components/InputTestActions.h"
+#include "EnhancedInputComponent.h"
+#include "EnhancedInputSubsystemInterface.h"
+#include "GameFramework/Pawn.h"
+#include "GameFramework/PlayerController.h"
+#include "UObject/StrongObjectPtr.h"
+#include "CQTestInputTestHelper.generated.h"
+
+class APawn;
+class APlayerController;
+class UEnhancedPlayerInput;
+class UInputMappingContext;
+
+/** Simple struct used for emulating a button interaction. Used for the InputActionTests. */
+struct FPressButtonAction : public FTestAction
+{
+	FPressButtonAction(const FString& ButtonName)
+	{
+		InputActionName = ButtonName;
+		InputActionValue = FInputActionValue(true);
+	}
+};
+
+/** Simple struct used for emulating axis interaction. Used for the InputActionTests. */
+struct FHoldAxisAction : public FTestAction
+{
+	FHoldAxisAction(const FString& AxisName, const FInputActionValue& ActionValue)
+	{
+		InputActionName = AxisName;
+		InputActionValue = ActionValue;
+	}
+};
+
+/** Mock input subsystems to avoid having to create an actual subsystem. */
+UCLASS(NotBlueprintable, NotBlueprintType)
+class UTestEnhancedInputSubsystem : public UObject, public IEnhancedInputSubsystemInterface
+{
+	GENERATED_BODY()
+
+	UEnhancedPlayerInput* PlayerInput;
+	UEnhancedInputUserSettings* UserSettings;
+
+public:
+	void Init(APlayerController* InPlayerController);
+	UEnhancedPlayerInput* GetPlayerInput() const override { return PlayerInput; }
+
+	// Records delegate triggering result, allowing tests to validate that they fired correctly
+	void MappingListener(const FInputActionInstance& Instance);
+
+	ETriggerEvent* GetEventForAction(const FString& InputActionName);
+
+protected:
+	virtual TMap<TObjectPtr<const UInputAction>, FInjectedInput>& GetContinuouslyInjectedInputs() override { return ContinuouslyInjectedInputs; }
+	
+	// Map of inputs that should be injected every frame. These inputs will be injected when ForcedInput is ticked.
+	UPROPERTY(Transient)
+	TMap<TObjectPtr<const UInputAction>, FInjectedInput> ContinuouslyInjectedInputs;
+
+private:
+	void SetupTestBindings();
+	void BindInputAction(const FString& InputActionName, EInputActionValueType InputActionValueType, const TArray<FKey>& Keys);
+
+	APlayerController* PlayerController;
+	UEnhancedInputComponent* InputComponent;
+	UInputMappingContext* InputMappingContext;
+	TMap<FString, ETriggerEvent> InputActionTriggerEventMap;
+};
+
+/** Helper class used for the InputActionTests to emulate an input system as the plugin may execute tests on a Pawn with no Player or input system created. */
+class FCQTestInputSubsystemHelper
+{
+public:
+	explicit FCQTestInputSubsystemHelper(APawn* InPawn) : Pawn(InPawn)
+	{
+		InitializePlayerControllerInput();
+	}
+
+	bool ActionExpectedEvent(const FString& InputActionName, ETriggerEvent ExpectedTriggerEvent);
+
+	inline static const FString TestButtonActionName = TEXT("TestButtonAction");
+	inline static const FString TestAxisActionName = TEXT("TestAxisAction");
+
+private:
+	void InitializePlayerControllerInput();
+
+	APawn* Pawn{ nullptr };
+	APlayerController* PlayerController{ nullptr };
+
+	TStrongObjectPtr<UTestEnhancedInputSubsystem> InputSubsystem{ nullptr };
+};
+
+/** Inherited InputTestAction used for testing our button and axis interactions. */
+class FCQTestPawnTestActions : public FInputTestActions
+{
+public:
+	explicit FCQTestPawnTestActions(APawn* Pawn) : FInputTestActions(Pawn)
+	{
+		InputSubsystemHelper = MakeUnique<FCQTestInputSubsystemHelper>(Pawn);
+	}
+
+	void PressButton(const FString& ButtonActionName);
+	void HoldAxis(const FString& ExampleAxisActionName, const FInputActionValue& ActionValue, FTimespan&& Duration);
+	bool IsTriggered(const FString& InputActionName);
+	bool IsCompleted(const FString& InputActionName);
+
+	FDateTime StartTime{ 0 };
+
+	// Because we're testing the input functionality within the plugin, we need to create and handle our own input system.
+	TUniquePtr<FCQTestInputSubsystemHelper> InputSubsystemHelper{ nullptr };
+};
\ No newline at end of file
diff --git a/Engine/Plugins/Tests/CQTest/Source/CQTestTests/Private/Components/InputActionTests.cpp b/Engine/Plugins/Tests/CQTest/Source/CQTestTests/Private/Components/InputActionTests.cpp
new file mode 100644
--- /dev/null
+++ b/Engine/Plugins/Tests/CQTest/Source/CQTestTests/Private/Components/InputActionTests.cpp
@@ -0,0 +1,49 @@
+// Copyright Epic Games, Inc. All Rights Reserved.
+
+#include "CQTest.h"
+
+#include "Components/MapTestSpawner.h"
+#include "Components/InputTestActions.h"
+#include "CQTestInputTestHelper.h"
+#include "EnhancedInputComponent.h"
+#include "EnhancedInputSubsystemInterface.h"
+#include "GameFramework/Pawn.h"
+
+#if WITH_EDITOR && WITH_AUTOMATION_TESTS
+
+TEST_CLASS(PawnActionTests, "TestFramework.CQTest.Input")
+{
+	APawn* TestPawn{ nullptr };
+	APlayerController* PlayerController{ nullptr };
+	TUniquePtr<FMapTestSpawner> Spawner{ nullptr };
+	TUniquePtr<FCQTestPawnTestActions> PawnActions{ nullptr };
+
+	BEFORE_EACH() {
+		Spawner = FMapTestSpawner::CreateFromTempLevel(TestCommandBuilder);
+		ASSERT_THAT(IsNotNull(Spawner));
+
+		Spawner->AddWaitUntilLoadedCommand(TestRunner);
+		TestCommandBuilder
+			.StartWhen([this]() {return nullptr != Spawner->FindFirstPlayerPawn(); })
+			.Then([this]() { PawnActions = MakeUnique<FCQTestPawnTestActions>(Spawner->FindFirstPlayerPawn()); });
+	}
+
+	TEST_METHOD(PawnAction_TestButtonPressAction)
+	{
+		TestCommandBuilder
+			.Do([this]() { PawnActions->PressButton(FCQTestInputSubsystemHelper::TestButtonActionName); })
+			.Then([this]() { ASSERT_THAT(IsTrue(PawnActions->IsTriggered(FCQTestInputSubsystemHelper::TestButtonActionName))); })
+			.Then([this]() { ASSERT_THAT(IsTrue(PawnActions->IsCompleted(FCQTestInputSubsystemHelper::TestButtonActionName))); });
+	}
+
+	TEST_METHOD(PawnAction_TestHoldAxisAction)
+	{
+		TestCommandBuilder
+			.Do([this]() { PawnActions->HoldAxis(FCQTestInputSubsystemHelper::TestAxisActionName, FInputActionValue(1.0f), FTimespan::FromMilliseconds(500)); })
+			.Then([this]() { ASSERT_THAT(IsTrue(PawnActions->IsTriggered(FCQTestInputSubsystemHelper::TestAxisActionName))); })
+			.Then([this]() { ASSERT_THAT(IsFalse(PawnActions->IsCompleted(FCQTestInputSubsystemHelper::TestAxisActionName))); })
+			.Until([this]() { return PawnActions->IsCompleted(FCQTestInputSubsystemHelper::TestAxisActionName); });
+	}
+};
+
+#endif // WITH_EDITOR && WITH_AUTOMATION_TESTS
\ No newline at end of file
