Prevent accidental floating point errors in CQTest
#tests Assertions and ObjectBuilder unit tested

diff --git a/Engine/Plugins/Tests/CQTest/Source/CQTest/Public/Assert/NoDiscardAsserter.inl b/Engine/Plugins/Tests/CQTest/Source/CQTest/Public/Assert/NoDiscardAsserter.inl
--- a/Engine/Plugins/Tests/CQTest/Source/CQTest/Public/Assert/NoDiscardAsserter.inl
+++ b/Engine/Plugins/Tests/CQTest/Source/CQTest/Public/Assert/NoDiscardAsserter.inl
@@ -130,12 +130,7 @@ template <typename T>
 template <typename TExpected, typename TActual>
 [[nodiscard]] FORCEINLINE bool FNoDiscardAsserter::AreEqual(const TExpected& Expected, const TActual& Actual)
 {
-	if (CQTestCondition::IsNotEqual(Expected, Actual))
-	{
-		Fail(FString::Printf(TEXT("Expected %s to equal %s."), *CQTestConvert::ToString(Expected), *CQTestConvert::ToString(Actual)));
-		return false;
-	}
-	return true;
+	return AreEqual(Expected, Actual, FString::Printf(TEXT("Expected %s to equal %s."), *CQTestConvert::ToString(Expected), *CQTestConvert::ToString(Actual)));
 }
 
 template <typename TExpected, typename TActual>
@@ -153,6 +148,9 @@ template <typename TExpected, typename TActual>
 template <typename TExpected, typename TActual>
 [[nodiscard]] FORCEINLINE bool FNoDiscardAsserter::AreEqual(const TExpected& Expected, const TActual& Actual, const FString& FailureMessage)
 {
+	static_assert(!std::is_floating_point<TExpected>::value && !std::is_floating_point<TActual>::value,
+		"Floating point types should not be compared for strict equality. Instead use IsNear() to account for floating point error. To directly compare without accounting for error, use IsTrue(A == B).");
+
 	if (CQTestCondition::IsNotEqual(Expected, Actual))
 	{
 		Fail(FailureMessage);
@@ -163,12 +161,7 @@ template <typename TExpected, typename TActual>
 
 [[nodiscard]] FORCEINLINE bool FNoDiscardAsserter::AreEqualIgnoreCase(const FString& Expected, const FString& Actual)
 {
-	if (CQTestCondition::IsNotEqualIgnoreCase(Expected, Actual))
-	{
-		Fail(FString::Printf(TEXT("Expected %s to equal %s ignoring case."), *Expected, *Actual));
-		return false;
-	}
-	return true;
+	return AreEqualIgnoreCase(Expected, Actual, FString::Printf(TEXT("Expected %s to equal %s ignoring case."), *Expected, *Actual));
 }
 
 [[nodiscard]] FORCEINLINE bool FNoDiscardAsserter::AreEqualIgnoreCase(const FString& Expected, const FString& Actual, const char* FailureMessage)
@@ -194,12 +187,7 @@ template <typename TExpected, typename TActual>
 template <typename TExpected, typename TActual>
 [[nodiscard]] FORCEINLINE bool FNoDiscardAsserter::AreNotEqual(const TExpected& Expected, const TActual& Actual)
 {
-	if (CQTestCondition::IsEqual(Expected, Actual))
-	{
-		Fail(FString::Printf(TEXT("Expected %s to not equal %s."), *CQTestConvert::ToString(Expected), *CQTestConvert::ToString(Actual)));
-		return false;
-	}
-	return true;
+	return AreNotEqual(Expected, Actual, FString::Printf(TEXT("Expected %s to not equal %s."), *CQTestConvert::ToString(Expected), *CQTestConvert::ToString(Actual)));
 }
 
 template <typename TExpected, typename TActual>
@@ -217,6 +205,9 @@ template <typename TExpected, typename TActual>
 template <typename TExpected, typename TActual>
 [[nodiscard]] FORCEINLINE bool FNoDiscardAsserter::AreNotEqual(const TExpected& Expected, const TActual& Actual, const FString& FailureMessage)
 {
+	static_assert(!std::is_floating_point<TExpected>::value && !std::is_floating_point<TActual>::value,
+		"Floating point types should not be compared for strict equality. Instead use !IsNear() to account for floating point error. To directly compare without accounting for error, use IsTrue(A != B).");
+
 	if (CQTestCondition::IsEqual(Expected, Actual))
 	{
 		Fail(FailureMessage);
@@ -227,12 +218,7 @@ template <typename TExpected, typename TActual>
 
 [[nodiscard]] FORCEINLINE bool FNoDiscardAsserter::AreNotEqualIgnoreCase(const FString& Expected, const FString& Actual)
 {
-	if (CQTestCondition::IsEqualIgnoreCase(Expected, Actual))
-	{
-		Fail(FString::Printf(TEXT("Expected %s to not equal %s ignoring case."), *Expected, *Actual));
-		return false;
-	}
-	return true;
+	return AreNotEqualIgnoreCase(Expected, Actual, FString::Printf(TEXT("Expected %s to not equal %s ignoring case."), *Expected, *Actual));
 }
 
 [[nodiscard]] FORCEINLINE bool FNoDiscardAsserter::AreNotEqualIgnoreCase(const FString& Expected, const FString& Actual, const char* FailureMessage)
@@ -258,12 +244,7 @@ template <typename TExpected, typename TActual>
 template <typename TExpected, typename TActual, typename TEpsilon>
 [[nodiscard]] FORCEINLINE bool FNoDiscardAsserter::IsNear(TExpected Expected, TActual Actual, TEpsilon Epsilon)
 {
-	if (CQTestCondition::IsNotNearlyEqual(Expected, Actual, Epsilon))
-	{
-		Fail(FString::Printf(TEXT("Expected %s to be near %s (within %s)."), *CQTestConvert::ToString(Expected), *CQTestConvert::ToString(Actual), *CQTestConvert::ToString(Epsilon)));
-		return false;
-	}
-	return true;
+	return IsNear(Expected, Actual, Epsilon, FString::Printf(TEXT("Expected %s to be near %s (within %s)."), *CQTestConvert::ToString(Expected), *CQTestConvert::ToString(Actual), *CQTestConvert::ToString(Epsilon)));
 }
 
 template <typename TExpected, typename TActual, typename TEpsilon>
diff --git a/Engine/Plugins/Tests/CQTest/Source/CQTestTests/Private/Assert/AssertionTests.cpp b/Engine/Plugins/Tests/CQTest/Source/CQTestTests/Private/Assert/AssertionTests.cpp
--- a/Engine/Plugins/Tests/CQTest/Source/CQTestTests/Private/Assert/AssertionTests.cpp
+++ b/Engine/Plugins/Tests/CQTest/Source/CQTestTests/Private/Assert/AssertionTests.cpp
@@ -143,28 +143,94 @@ namespace CQTests
 			Assert.ExpectError(ExpectedError);
 			ASSERT_THAT(AreEqual(42, 0, ExpectedError));
 		}
-		TEST_METHOD(AssertNear_WithSameNumbers_Succeeds)
+
+		/* These cases intentionally do not compile to avoid accidental floating point error
+		TEST_METHOD(AssertEqual_WithDoubles_CompileError)
+		{
+			ASSERT_THAT(AreEqual(42.0, 42.0));
+		}
+		TEST_METHOD(AssertEqual_WithFloats_CompileError)
+		{
+			ASSERT_THAT(AreEqual(42.f, 42.f));
+		}
+		TEST_METHOD(AssertEqual_WithMixedFloat_CompileError)
+		{
+			ASSERT_THAT(AreEqual(42, 42.f));
+		}
+
+		TEST_METHOD(AssertNotEqual_WithDoubles_CompileError)
+		{
+			ASSERT_THAT(AreNotEqual(42.0, 42.0));
+		}
+		TEST_METHOD(AssertNotEqual_WithFloats_CompileError)
+		{
+			ASSERT_THAT(AreNotEqual(42.f, 42.f));
+		}
+		TEST_METHOD(AssertNotEqual_WithMixedFloat_CompileError)
+		{
+			ASSERT_THAT(AreNotEqual(42, 42.f));
+		}
+		*/
+
+		TEST_METHOD(AssertNotEqual_WithDifferentInts_Succeeds)
+		{
+			ASSERT_THAT(AreNotEqual(42, 0));
+		}
+		TEST_METHOD(AssertNotEqual_WithDifferentAndErrorMessage_DoesNotAddErrorMessage)
+		{
+			ASSERT_THAT(AreNotEqual(42, 0, "Unexpected"));
+		}
+		TEST_METHOD(AssertNotEqual_WithSameInts_AddsError)
+		{
+			Assert.ExpectError(AnyError);
+			ASSERT_THAT(AreNotEqual(42, 42));
+		}
+		TEST_METHOD(AssertNotEqual_WithSameIntsAndErrorMessage_AddsSpecificError)
+		{
+			Assert.ExpectError(ExpectedError);
+			ASSERT_THAT(AreNotEqual(42, 42, ExpectedError));
+		}
+
+		TEST_METHOD(AssertNear_WithSameDoubles_Succeeds)
 		{
 			ASSERT_THAT(IsNear(3.14, 3.14, 0.001));
 		}
-		TEST_METHOD(AssertNear_WithSameNumbersAndErrorMessage_DoesNotAddErrorMessage)
+		TEST_METHOD(AssertNear_WithSameDoublesAndErrorMessage_DoesNotAddErrorMessage)
 		{
 			ASSERT_THAT(IsNear(3.14, 3.14, 0.001, "Unexpected"));
 		}
-		TEST_METHOD(AssertNear_WithSimilarNumbers_Succeeds)
+		TEST_METHOD(AssertNear_WithSimilarDoubles_Succeeds)
 		{
 			ASSERT_THAT(IsNear(3.0, 3.1, 1.0));
 		}
-		TEST_METHOD(AssertNear_WithDifferentNumbers_AddsError)
+		TEST_METHOD(AssertNear_WithDifferentDoubles_AddsError)
 		{
 			Assert.ExpectError(AnyError);
 			ASSERT_THAT(IsNear(1.0, 2.0, 0.001));
 		}
-		TEST_METHOD(AssertNear_WithDifferentNumbersAndError_AddsSpecificError)
+		TEST_METHOD(AssertNear_WithDifferentDoublesAndError_AddsSpecificError)
 		{
 			Assert.ExpectError(ExpectedError);
 			ASSERT_THAT(IsNear(1.0, 2.0, 0.001, ExpectedError));
 		}
+		TEST_METHOD(AssertNear_WithSimilarDoubleWithinTolerance_Succeeds)
+		{
+			ASSERT_THAT(IsNear(3.0, 3.001, 0.01));
+		}
+		TEST_METHOD(AssertNear_WithSimilarDoubleWithinNegativeTolerance_Succeeds)
+		{
+			ASSERT_THAT(IsNear(3.0, 2.999, 0.01));
+		}
+		TEST_METHOD(AssertNear_WithSimilarDoubleOutsideTolerance_AddsError)
+		{
+			Assert.ExpectError(AnyError);
+			ASSERT_THAT(IsNear(3.0, 3.1, 0.01));
+		}
+		TEST_METHOD(AssertNear_WithSimilarDoubleOutsideNegativeTolerance_AddsError)
+		{
+			Assert.ExpectError(AnyError);
+			ASSERT_THAT(IsNear(3.0, 2.9, 0.01));
+		}
 	};
 
 	TEST_CLASS(NoDiscardAssert_Pointers, "TestFramework.CQTest.Core")
diff --git a/Engine/Plugins/Tests/CQTest/Source/CQTestTests/Private/ObjectBuilderTests.cpp b/Engine/Plugins/Tests/CQTest/Source/CQTestTests/Private/ObjectBuilderTests.cpp
--- a/Engine/Plugins/Tests/CQTest/Source/CQTestTests/Private/ObjectBuilderTests.cpp
+++ b/Engine/Plugins/Tests/CQTest/Source/CQTestTests/Private/ObjectBuilderTests.cpp
@@ -109,16 +109,16 @@ TEST_CLASS(ObjectBuilder_Success, "TestFramework.CQTest.Actor")
 						   .SetParam("DoubleProperty", doubleValue)
 						   .Spawn();
 
-		ASSERT_THAT(AreEqual(Result.FloatProperty, floatValue));
-		ASSERT_THAT(AreEqual(Result.DoubleProperty, doubleValue));
+		ASSERT_THAT(IsTrue(Result.FloatProperty == floatValue));
+		ASSERT_THAT(IsTrue(Result.DoubleProperty == doubleValue));
 
 		auto& ConstResult = TObjectBuilder<ATestActorWithProperties>(Spawner)
 						   .SetParam<const float>("FloatProperty", floatValue)
 						   .SetParam<const double>("DoubleProperty", doubleValue)
 						   .Spawn();
 
-		ASSERT_THAT(AreEqual(ConstResult.FloatProperty, floatValue));
-		ASSERT_THAT(AreEqual(ConstResult.DoubleProperty, doubleValue));
+		ASSERT_THAT(IsTrue(Result.FloatProperty == floatValue));
+		ASSERT_THAT(IsTrue(Result.DoubleProperty == doubleValue));
 	}
 
 	TEST_METHOD(Spawn_ActorWithEnums_SetsEnums)
