Added FWaitDelay latent command to allow for fixed duration waits.

diff --git a/Engine/Plugins/Tests/CQTest/Source/CQTest/Private/Commands/TestCommands.cpp b/Engine/Plugins/Tests/CQTest/Source/CQTest/Private/Commands/TestCommands.cpp
--- a/Engine/Plugins/Tests/CQTest/Source/CQTest/Private/Commands/TestCommands.cpp
+++ b/Engine/Plugins/Tests/CQTest/Source/CQTest/Private/Commands/TestCommands.cpp
@@ -45,6 +45,33 @@ bool CQTEST_API FWaitUntil::Update()
 	return false;
 }
 
+bool CQTEST_API FWaitDelay::Update()
+{
+	if (!bHasTimerStarted)
+	{
+		EndTime = FDateTime::UtcNow() + Timeout;
+		bHasTimerStarted = true;
+		if (Description != nullptr)
+		{
+			UE_LOG(LogCqTest, Log, TEXT("Starting %s"), Description);
+		}
+	}
+	if (TestRunner.HasAnyErrors())
+	{
+		return true;
+	}
+
+	if (FDateTime::UtcNow() >= EndTime)
+	{
+		if (Description)
+		{
+			UE_LOG(LogCqTest, Log, TEXT("Finished %s"), Description);
+		}
+		return true;
+	}
+	return false;
+}
+
 bool CQTEST_API FExecute::Update()
 {
 	bool bHasToRun = !TestRunner.HasAnyErrors() || FailureBehavior == ECQTestFailureBehavior::Run;
diff --git a/Engine/Plugins/Tests/CQTest/Source/CQTest/Public/Commands/TestCommandBuilder.h b/Engine/Plugins/Tests/CQTest/Source/CQTest/Public/Commands/TestCommandBuilder.h
--- a/Engine/Plugins/Tests/CQTest/Source/CQTest/Public/Commands/TestCommandBuilder.h
+++ b/Engine/Plugins/Tests/CQTest/Source/CQTest/Public/Commands/TestCommandBuilder.h
@@ -65,6 +65,22 @@ public:
 		return Until(Description, Query, Timeout);
 	}
 
+	/** Note that using a timed-wait can introduce test flakiness due to variable runtimes. Please consider using `Until` and waiting until something happens instead. */
+	FTestCommandBuilder& WaitDelay(FTimespan Timeout)
+	{
+		return WaitDelay(nullptr, Timeout);
+	}
+
+	/** Note that using a timed-wait can introduce test flakiness due to variable runtimes. Please consider using `Until` and waiting until something happens instead. */
+	FTestCommandBuilder& WaitDelay(const TCHAR* Description, FTimespan Timeout)
+	{
+		if (!TestRunner.HasAnyErrors())
+		{
+			CommandQueue.Add(MakeShared<FWaitDelay>(TestRunner, Timeout, Description));
+		}
+		return *this;
+	}
+
 	FTestCommandBuilder& OnTearDown(const TCHAR* Description, TFunction<void()> Action)
 	{
 		if (!TestRunner.HasAnyErrors())
diff --git a/Engine/Plugins/Tests/CQTest/Source/CQTest/Public/Commands/TestCommands.h b/Engine/Plugins/Tests/CQTest/Source/CQTest/Public/Commands/TestCommands.h
--- a/Engine/Plugins/Tests/CQTest/Source/CQTest/Public/Commands/TestCommands.h
+++ b/Engine/Plugins/Tests/CQTest/Source/CQTest/Public/Commands/TestCommands.h
@@ -24,6 +24,27 @@ public:
 	bool bHasTimerStarted = false;
 };
 
+/** Latent Command that waits a set time frame.
+ * Note that using a timed-wait can introduce test flakiness due to variable runtimes. Please consider using `FWaitUntil` and waiting until something happens instead.
+ */
+class CQTEST_API FWaitDelay : public IAutomationLatentCommand
+{
+public:
+	FWaitDelay(FAutomationTestBase& InTestRunner, FTimespan Timeout, const TCHAR* InDescription = nullptr)
+		: TestRunner(InTestRunner)
+		, Timeout(Timeout)
+		, Description(InDescription)
+	{}
+
+	bool Update() override;
+
+	FAutomationTestBase& TestRunner;
+	FTimespan Timeout;
+	FDateTime EndTime;
+	const TCHAR* Description;
+	bool bHasTimerStarted = false;
+};
+
 enum class ECQTestFailureBehavior
 {
 	Skip,
diff --git a/Engine/Plugins/Tests/CQTest/Source/CQTest/Public/Commands/TestMacros.h b/Engine/Plugins/Tests/CQTest/Source/CQTest/Public/Commands/TestMacros.h
--- a/Engine/Plugins/Tests/CQTest/Source/CQTest/Public/Commands/TestMacros.h
+++ b/Engine/Plugins/Tests/CQTest/Source/CQTest/Public/Commands/TestMacros.h
@@ -9,5 +9,6 @@
 
 #define DO(Cmd) TestCommandBuilder.Do([&]() { Cmd })
 #define START_WHEN(Query) TestCommandBuilder.StartWhen([&]() { Query })
+#define WAIT_DELAY(Duration) TestCommandBuilder.WaitDelay(Duration)
 #define THEN(Cmd) Then([&]() { Cmd })
 #define UNTIL(Query) Until([&]() { Query })
\ No newline at end of file
diff --git a/Engine/Plugins/Tests/CQTest/Source/CQTestTests/Private/Commands/CommandBuilderTests.cpp b/Engine/Plugins/Tests/CQTest/Source/CQTestTests/Private/Commands/CommandBuilderTests.cpp
--- a/Engine/Plugins/Tests/CQTest/Source/CQTestTests/Private/Commands/CommandBuilderTests.cpp
+++ b/Engine/Plugins/Tests/CQTest/Source/CQTestTests/Private/Commands/CommandBuilderTests.cpp
@@ -1,6 +1,7 @@
 // Copyright Epic Games, Inc. All Rights Reserved.
 
 #include "CQTest.h"
+#include "CQTestUnitTestHelper.h"
 
 TEST_CLASS(CommandBuilderTests, "TestFramework.CQTest.Core")
 {
@@ -31,6 +32,43 @@ TEST_CLASS(CommandBuilderTests, "TestFramework.CQTest.Core")
 		ASSERT_THAT(IsTrue(command->Update()));
 	}
 
+	TEST_METHOD(WaitDelay_WaitsUntilDurationElapsed)
+	{
+		bool done = false;
+		FTimespan Duration = FTimespan::FromMilliseconds(200);
+		FDateTime EndTime = FDateTime::UtcNow() + Duration;
+		auto command = CommandBuilder
+			.WaitDelay(Duration)
+			.Then([this, &EndTime, &done]()
+			{
+				ASSERT_THAT(IsTrue(FDateTime::UtcNow() >= EndTime));
+				done = true;
+			}).Build();
+		
+		while (!done)
+		{
+			command->Update();
+		}
+		ASSERT_THAT(IsTrue(done));
+	}
+
+	TEST_METHOD(WaitDelay_InterruptOnError)
+	{
+		const FString ExpectedError = TEXT("Error reported outside WaitDelay");
+
+		FTimespan Duration = FTimespan::FromSeconds(10);
+		FDateTime EndTime = FDateTime::UtcNow() + Duration;
+		auto command = CommandBuilder
+			.WaitDelay(Duration).Build();
+
+		ASSERT_THAT(IsFalse(command->Update()));
+		AddError(ExpectedError);
+		ASSERT_THAT(IsTrue(command->Update()));
+		ASSERT_THAT(IsTrue(FDateTime::UtcNow() < EndTime));
+
+		ClearExpectedError(*this->TestRunner, ExpectedError);
+	}
+
 	TEST_METHOD(Build_AfterBuild_ReturnsNullptr)
 	{
 		auto command = CommandBuilder.Do([]() {}).Build();
