Allow FMapTestSpawner to be used across multiple build types and not just limited to the Editor

#rnx

diff --git a/Engine/Plugins/Tests/CQTest/Source/CQTest/Private/Components/MapTestSpawner.cpp b/Engine/Plugins/Tests/CQTest/Source/CQTest/Private/Components/MapTestSpawner.cpp
--- a/Engine/Plugins/Tests/CQTest/Source/CQTest/Private/Components/MapTestSpawner.cpp
+++ b/Engine/Plugins/Tests/CQTest/Source/CQTest/Private/Components/MapTestSpawner.cpp
@@ -2,12 +2,16 @@
 
 #include "Components/MapTestSpawner.h"
 
-#if ENABLE_MAPSPAWNER_TEST
+#if WITH_AUTOMATION_TESTS
 #include "Commands/TestCommands.h"
 #include "Tests/AutomationCommon.h"
 #include "GameFramework/PlayerController.h"
 #include "Engine/Engine.h"
+#include "Misc/PackageName.h"
 #include "Misc/Paths.h"
+#include "UObject/Package.h"
+
+#if WITH_EDITOR
 #include "Editor.h"
 #include "Editor/UnrealEdEngine.h"
 #include "HAL/FileManager.h"
@@ -45,16 +49,11 @@ void CleanupTempResources()
 }
 
 } //anonymous
-
-FMapTestSpawner::FMapTestSpawner(const FString& MapDirectory, const FString& MapName) : MapDirectory(MapDirectory) , MapName(MapName)
-{
-	// Register Map Change Events
-	FLevelEditorModule& LevelEditor = FModuleManager::LoadModuleChecked<FLevelEditorModule>("LevelEditor");
-	MapChangedHandle = LevelEditor.OnMapChanged().AddRaw(this, &FMapTestSpawner::OnMapChanged);
-}
+#endif // WITH_EDITOR
 
 TUniquePtr<FMapTestSpawner> FMapTestSpawner::CreateFromTempLevel(FTestCommandBuilder& InCommandBuilder)
 {
+#if WITH_EDITOR
 	if (IsValid(GUnrealEd->PlayWorld))
 	{
 		UE_LOG(LogMapTest, Verbose, TEXT("Active PIE session '%s' needs to be shutdown before a creation of a new level can occur."), *GUnrealEd->PlayWorld->GetMapName());
@@ -70,12 +69,21 @@ TUniquePtr<FMapTestSpawner> FMapTestSpawner::CreateFromTempLevel(FTestCommandBui
 	check(bWasTempLevelCreated);
 
 	TUniquePtr<FMapTestSpawner> Spawner = MakeUnique<FMapTestSpawner>(TempMapDirectory, MapName);
+	
+	// Register Map Change Events
+	FLevelEditorModule& LevelEditor = FModuleManager::LoadModuleChecked<FLevelEditorModule>("LevelEditor");
+	Spawner->MapChangedHandle = LevelEditor.OnMapChanged().AddRaw(Spawner.Get(), &FMapTestSpawner::OnMapChanged);
+
 	InCommandBuilder.OnTearDown([&]() {
 		// Create a new map to free up the reference to the map used during testing before cleaning up all temporary resources
 		FAutomationEditorCommonUtils::CreateNewMap();
 		CleanupTempResources();
 	});
 	return MoveTemp(Spawner);
+#else
+	checkf(false, TEXT("CreateFromTempLevel can't create a new level if WITH_EDITOR=false"));
+	return nullptr;
+#endif // WITH_EDITOR
 }
 
 void FMapTestSpawner::AddWaitUntilLoadedCommand(FAutomationTestBase* TestRunner)
@@ -85,15 +93,33 @@ void FMapTestSpawner::AddWaitUntilLoadedCommand(FAutomationTestBase* TestRunner)
 	FString PackagePath;
 	const FString Path = FPaths::Combine(MapDirectory, MapName);
 	bool bPackageExists = FPackageName::DoesPackageExist(Path, &PackagePath);
-	check(bPackageExists);
+	checkf(bPackageExists, TEXT("Could not get package from path '%s'"), *Path);
+
+	// We need to retrieve the LongPackageName from the PackagePath to be able to load the map for both Editor and Target builds
+	FString LongPackageName, PackageConversionError;
+	bool bFilenameConverted = FPackageName::TryConvertFilenameToLongPackageName(PackagePath, LongPackageName, &PackageConversionError);
+	checkf(bFilenameConverted, TEXT("Could not get LongPackageName. Error: '%s'"), *PackageConversionError);
 
-	bool bOpened = AutomationOpenMap(PackagePath);
+	bool bOpened = AutomationOpenMap(LongPackageName, true);
 	check(bOpened);
 
 	ADD_LATENT_AUTOMATION_COMMAND(FWaitUntil(*TestRunner, [&]() -> bool {
 		for (const auto& Context : GEngine->GetWorldContexts())
 		{
-			if (((Context.WorldType == EWorldType::PIE) || (Context.WorldType == EWorldType::Game)) && (Context.World() != nullptr))
+			UWorld* World = Context.World();
+			if (!IsValid(World))
+			{
+				continue;
+			}
+
+			// We only want to set our PieWorld if the loaded World name matches our expected World name
+			FString WorldMapName = FPackageName::GetShortName(World->GetMapName());
+			if (World->GetOutermost()->GetPIEInstanceID() != INDEX_NONE)
+			{
+				FString PIEPrefix = FString::Printf(PLAYWORLD_PACKAGE_PREFIX TEXT("_%d_"), World->GetOutermost()->GetPIEInstanceID());
+				WorldMapName.ReplaceInline(*PIEPrefix, TEXT(""));
+			}
+			if (((Context.WorldType == EWorldType::PIE) || (Context.WorldType == EWorldType::Game)) && (WorldMapName.Equals(MapName)))
 			{
 				PieWorld = Context.World();
 				return true;
@@ -112,9 +138,18 @@ UWorld* FMapTestSpawner::CreateWorld()
 
 APawn* FMapTestSpawner::FindFirstPlayerPawn()
 {
-	return GetWorld().GetFirstPlayerController()->GetPawn();
+	APlayerController* PlayerController = GetWorld().GetFirstPlayerController();
+
+	// There's a chance that we may not have a PlayerController spawned in the world
+	if (!IsValid(PlayerController))
+	{
+		return nullptr;
+	}
+
+	return PlayerController->GetPawn();
 }
 
+#if WITH_EDITOR
 void FMapTestSpawner::OnMapChanged(UWorld* World, EMapChangeType ChangeType)
 {
 	if (PieWorld && ChangeType == EMapChangeType::TearDownWorld)
@@ -127,5 +162,6 @@ void FMapTestSpawner::OnMapChanged(UWorld* World, EMapChangeType ChangeType)
 		LevelEditor.OnMapChanged().Remove(MapChangedHandle);
 	}
 }
+#endif // WITH_EDITOR
 
-#endif // ENABLE_MAPSPAWNER_TEST
\ No newline at end of file
+#endif // WITH_AUTOMATION_TESTS
\ No newline at end of file
diff --git a/Engine/Plugins/Tests/CQTest/Source/CQTest/Public/Components/MapTestSpawner.h b/Engine/Plugins/Tests/CQTest/Source/CQTest/Public/Components/MapTestSpawner.h
--- a/Engine/Plugins/Tests/CQTest/Source/CQTest/Public/Components/MapTestSpawner.h
+++ b/Engine/Plugins/Tests/CQTest/Source/CQTest/Public/Components/MapTestSpawner.h
@@ -6,10 +6,11 @@
 
 #include "Commands/TestCommandBuilder.h"
 
-#define ENABLE_MAPSPAWNER_TEST WITH_EDITOR && WITH_AUTOMATION_TESTS
+#if WITH_AUTOMATION_TESTS
 
-#if ENABLE_MAPSPAWNER_TEST
+#if WITH_EDITOR
 #include "UnrealEdMisc.h"
+#endif // WITH_EDITOR
 
 /// Class for spawning Actors in a named Map/Level
 struct CQTEST_API FMapTestSpawner : public FSpawnHelper
@@ -20,7 +21,8 @@ struct CQTEST_API FMapTestSpawner : public FSpawnHelper
 	 * @param MapDirectory - The directory which the map resides in.
 	 * @param MapName - Name of the map.
 	 */
-	FMapTestSpawner(const FString& MapDirectory, const FString& MapName);
+	FMapTestSpawner(const FString& MapDirectory, const FString& MapName)
+		: MapDirectory(MapDirectory), MapName(MapName) {}
 
 	/**
 	 * Creates an instance of the MapTestSpawner with a temporary level ready for use.
@@ -46,15 +48,18 @@ protected:
     virtual UWorld* CreateWorld() override;
 
 private:
+#if WITH_EDITOR
 	/**
 	 * Handler called on map changed.
 	 */
 	void OnMapChanged(UWorld* World, EMapChangeType ChangeType);
 
+	FDelegateHandle MapChangedHandle;
+#endif // WITH_EDITOR
+
 	FString MapDirectory;
 	FString MapName;
 	UWorld* PieWorld{ nullptr };
-	FDelegateHandle MapChangedHandle;
 };
 
-#endif // ENABLE_MAPSPAWNER_TEST
\ No newline at end of file
+#endif // WITH_AUTOMATION_TESTS
\ No newline at end of file
diff --git a/Engine/Plugins/Tests/CQTest/Source/CQTestTests/Private/Components/MapSpawnHelperTests.cpp b/Engine/Plugins/Tests/CQTest/Source/CQTestTests/Private/Components/MapSpawnHelperTests.cpp
--- a/Engine/Plugins/Tests/CQTest/Source/CQTestTests/Private/Components/MapSpawnHelperTests.cpp
+++ b/Engine/Plugins/Tests/CQTest/Source/CQTestTests/Private/Components/MapSpawnHelperTests.cpp
@@ -4,9 +4,9 @@
 #include "Components/MapTestSpawner.h"
 #include "GameFramework/Pawn.h"
 
-#if ENABLE_MAPSPAWNER_TEST
+#if WITH_EDITOR && WITH_AUTOMATION_TESTS
 
-TEST_CLASS(MapSpawnHelperTests, "TestFramework.CQTest.Map")
+TEST_CLASS_WITH_FLAGS(MapSpawnHelperTests, "TestFramework.CQTest.Map", EAutomationTestFlags::EditorContext | EAutomationTestFlags::EngineFilter)
 {
 	TUniquePtr<FMapTestSpawner> Spawner;
 
@@ -30,4 +30,4 @@ TEST_CLASS(MapSpawnHelperTests, "TestFramework.CQTest.Map")
 	}
 };
 
-#endif // ENABLE_MAPSPAWNER_TEST
\ No newline at end of file
+#endif // WITH_EDITOR && WITH_AUTOMATION_TESTS
\ No newline at end of file
