Added `CQTestSlateComponent` which starts introducing Slate into CQTest. Currently, only checking if the UI was refreshed/ticked

#rnx

diff --git a/Engine/Plugins/Tests/CQTest/Source/CQTest/CQTest.Build.cs b/Engine/Plugins/Tests/CQTest/Source/CQTest/CQTest.Build.cs
--- a/Engine/Plugins/Tests/CQTest/Source/CQTest/CQTest.Build.cs
+++ b/Engine/Plugins/Tests/CQTest/Source/CQTest/CQTest.Build.cs
@@ -16,7 +16,8 @@ public class CQTest : ModuleRules
 					"Core",
 					"CoreUObject",
 					"Engine",
-					"EnhancedInput"
+					"EnhancedInput",
+					"Slate"
 				 }
 			);
 
diff --git a/Engine/Plugins/Tests/CQTest/Source/CQTest/Private/Components/CQTestSlateComponent.cpp b/Engine/Plugins/Tests/CQTest/Source/CQTest/Private/Components/CQTestSlateComponent.cpp
new file mode 100644
--- /dev/null
+++ b/Engine/Plugins/Tests/CQTest/Source/CQTest/Private/Components/CQTestSlateComponent.cpp
@@ -0,0 +1,60 @@
+// Copyright Epic Games, Inc. All Rights Reserved.
+
+#include "Components/CQTestSlateComponent.h"
+
+#include "Framework/Application/SlateApplication.h"
+#include "Tests/AutomationCommon.h"
+
+DEFINE_LOG_CATEGORY_STATIC(LogCQTestSlateComponent, Log, All);
+
+FCQTestSlateComponent::FCQTestSlateComponent()
+{
+	checkf(FSlateApplication::IsInitialized(), TEXT("No Slate application initialized."));
+
+	// Disable Slate from going into a sleep state forcing it to always tick
+	TestEnvironment = FScopedTestEnvironment::Get();
+	TestEnvironment->SetConsoleVariableValue(TEXT("Slate.AllowSlateToSleep"), TEXT("0"));
+
+	TickDelegateHandle = FSlateApplication::Get().OnPostTick().AddRaw(this, &FCQTestSlateComponent::OnPostTick);
+}
+
+FCQTestSlateComponent::~FCQTestSlateComponent()
+{
+	if (FSlateApplication::IsInitialized() && TickDelegateHandle.IsValid())
+	{
+		FSlateApplication::Get().OnPostTick().Remove(TickDelegateHandle);
+	}
+}
+
+bool FCQTestSlateComponent::HaveTicksElapsed(uint32 Ticks)
+{
+	// Early out checking that 0 ticks have elapsed as there is nothing to wait for
+	if (Ticks == 0)
+	{
+		UE_LOG(LogCQTestSlateComponent, Verbose, TEXT("Nothing to wait for as the expected elapsed ticks requested is 0."));
+		return true;
+	}
+
+	if (!ExpectedTick.IsSet())
+	{
+		ExpectedTick = TickCounter + Ticks;
+
+		UE_LOG(LogCQTestSlateComponent, Verbose, TEXT("HaveTicksElapsed called for %d ticks."), Ticks);
+		UE_LOG(LogCQTestSlateComponent, Verbose, TEXT("Slate has ticked %d times."), TickCounter.load());
+		UE_LOG(LogCQTestSlateComponent, Verbose, TEXT("ExpectedTick set with a value of %d."), ExpectedTick.GetValue());
+	}
+	else if (ExpectedTick.GetValue() <= TickCounter)
+	{
+		UE_LOG(LogCQTestSlateComponent, Verbose, TEXT("Slate ticked %d times and has met the expected tick goal of %d ticks."), TickCounter.load(), ExpectedTick.GetValue());
+
+		ExpectedTick.Reset();
+		return true;
+	}
+
+	return false;
+}
+
+void FCQTestSlateComponent::OnPostTick(const float InDeltaTime)
+{
+	TickCounter++;
+}
\ No newline at end of file
diff --git a/Engine/Plugins/Tests/CQTest/Source/CQTest/Public/Components/CQTestSlateComponent.h b/Engine/Plugins/Tests/CQTest/Source/CQTest/Public/Components/CQTestSlateComponent.h
new file mode 100644
--- /dev/null
+++ b/Engine/Plugins/Tests/CQTest/Source/CQTest/Public/Components/CQTestSlateComponent.h
@@ -0,0 +1,47 @@
+// Copyright Epic Games, Inc. All Rights Reserved.
+
+#pragma once
+
+#include "Delegates/IDelegateInstance.h"
+#include "Misc/Optional.h"
+#include "Templates/SharedPointer.h"
+
+#include <atomic>
+
+// CQTest component for interacting with Slate
+class CQTEST_API FCQTestSlateComponent
+{
+public:
+	FCQTestSlateComponent();
+	~FCQTestSlateComponent();
+
+	/**
+	 * Initializes the ExpectedTick variable to be checked against the amount of times Slate has ticked.
+	 * Will be used to check Slate has ticked the expected amount before confirming and resetting the ExpectedTick.
+	 * 
+	 * Note that this method is meant to be used as the only statement within a latent command that waits until the predicate is met, such as FWaitUntil, Until, or StartWhen methods.
+	 *
+	 * @param Ticks - Number of ticks we're expecting to have elapsed
+	 * @return true when Slate has ticked the provided number of times.
+	 */
+	bool HaveTicksElapsed(uint32 Ticks);
+
+	/** Returns the total elapsed ticks since creation. */
+	uint32 GetTotalElapsedTicks()
+	{
+		return TickCounter;
+	}
+
+private:
+	/**
+	 * Called on Slate's Post Tick
+	 *
+	 * @param InDeltaTime - Elapsed time since the last tick. While the value is not being used directly, it is needed to register for Slate's delegate.
+	 */
+	void OnPostTick([[maybe_unused]] const float InDeltaTime);
+
+	std::atomic<uint32> TickCounter{0};
+	TOptional<uint32> ExpectedTick;
+	FDelegateHandle TickDelegateHandle;
+	TSharedPtr<struct FScopedTestEnvironment> TestEnvironment;
+};
\ No newline at end of file
diff --git a/Engine/Plugins/Tests/CQTest/Source/CQTestTests/CQTestTests.Build.cs b/Engine/Plugins/Tests/CQTest/Source/CQTestTests/CQTestTests.Build.cs
--- a/Engine/Plugins/Tests/CQTest/Source/CQTestTests/CQTestTests.Build.cs
+++ b/Engine/Plugins/Tests/CQTest/Source/CQTestTests/CQTestTests.Build.cs
@@ -16,6 +16,7 @@ public class CQTestTests : ModuleRules
 					"Engine",
 					"InputCore",
 					"EnhancedInput",
+					"Slate",
 					"CQTest"
 				 }
 			);
diff --git a/Engine/Plugins/Tests/CQTest/Source/CQTestTests/Private/Components/CQTestSlateComponentTests.cpp b/Engine/Plugins/Tests/CQTest/Source/CQTestTests/Private/Components/CQTestSlateComponentTests.cpp
new file mode 100644
--- /dev/null
+++ b/Engine/Plugins/Tests/CQTest/Source/CQTestTests/Private/Components/CQTestSlateComponentTests.cpp
@@ -0,0 +1,43 @@
+// Copyright Epic Games, Inc. All Rights Reserved.
+
+#include "CQTest.h"
+
+#include "Components/CQTestSlateComponent.h"
+#include "Framework/Application/SlateApplication.h"
+
+TEST_CLASS(SlateTests, "TestFramework.CQTest.UI")
+{
+	TUniquePtr<FCQTestSlateComponent> SlateComponent{ nullptr };
+
+	BEFORE_EACH()
+	{
+		SlateComponent = MakeUnique<FCQTestSlateComponent>();
+	}
+
+	TEST_METHOD(HaveTicksElapsed_WithoutTicking_ReturnsFalse)
+	{
+		ASSERT_THAT(IsTrue(SlateComponent->HaveTicksElapsed(0)));
+	}
+
+	TEST_METHOD(HaveTicksElapsed_AfterTicking_ReturnsTrue)
+	{
+		ASSERT_THAT(IsFalse(SlateComponent->HaveTicksElapsed(1)));
+		FSlateApplication::Get().Tick();
+		ASSERT_THAT(IsTrue(SlateComponent->HaveTicksElapsed(1)));
+	}
+
+	TEST_METHOD(HaveTicksElapsed_InUntilCommand_EventuallyReturnsTrue)
+	{
+		TestCommandBuilder
+			.StartWhen([this]() { return SlateComponent->HaveTicksElapsed(3); })
+			.Then([this]() { ASSERT_THAT(IsTrue(SlateComponent->GetTotalElapsedTicks() >= 3)); });
+	}
+
+	TEST_METHOD(ElapsedTicks_WhenReused_AwaitsAllTheTicks)
+	{
+		TestCommandBuilder
+			.StartWhen([this]() { return SlateComponent->HaveTicksElapsed(3); })
+			.Until([this]() { return SlateComponent->HaveTicksElapsed(3); })
+			.Then([this]() { ASSERT_THAT(IsTrue(SlateComponent->GetTotalElapsedTicks() >= 6)); });
+	}
+};
\ No newline at end of file
diff --git a/Engine/Source/Runtime/Engine/Private/Tests/AutomationCommon.cpp b/Engine/Source/Runtime/Engine/Private/Tests/AutomationCommon.cpp
--- a/Engine/Source/Runtime/Engine/Private/Tests/AutomationCommon.cpp
+++ b/Engine/Source/Runtime/Engine/Private/Tests/AutomationCommon.cpp
@@ -302,6 +302,7 @@ void FTestConsoleVariable::Set(const FString& Value)
 		}
 
 		ConsoleVariable->AsVariable()->SetWithCurrentPriority(*Value);
+		UE_LOG(LogEngineAutomationTests, Verbose, TEXT("CVar '%s' was set with the value '%s'."), *ConsoleVariableName, *Value);
 	}
 }
 
@@ -324,6 +325,7 @@ void FTestConsoleVariable::Restore()
 		if (ensure(ConsoleVariable))
 		{
 			ConsoleVariable->AsVariable()->SetWithCurrentPriority(*OriginalValue);
+			UE_LOG(LogEngineAutomationTests, Verbose, TEXT("CVar '%s' was restored with the original value '%s'."), *ConsoleVariableName, *OriginalValue);
 		}
 
 		bModified = false;
diff --git a/Engine/Source/Runtime/Engine/Public/Tests/AutomationCommon.h b/Engine/Source/Runtime/Engine/Public/Tests/AutomationCommon.h
--- a/Engine/Source/Runtime/Engine/Public/Tests/AutomationCommon.h
+++ b/Engine/Source/Runtime/Engine/Public/Tests/AutomationCommon.h
@@ -103,7 +103,7 @@ private:
 };
 
 /**
- * Utility for setting and management of the test environment
+ * Utility for setting and management of temporary CVars
  * Will handle restoring the Console Variables (CVars) back to the original state on destruction.
  */
 struct FScopedTestEnvironment
