CQTest support for BeforeAll and AfterAll

#rnx
#pr 11806

diff --git a/Engine/Plugins/Tests/CQTest/Source/CQTest/Public/CQTest.h b/Engine/Plugins/Tests/CQTest/Source/CQTest/Public/CQTest.h
--- a/Engine/Plugins/Tests/CQTest/Source/CQTest/Public/CQTest.h
+++ b/Engine/Plugins/Tests/CQTest/Source/CQTest/Public/CQTest.h
@@ -19,8 +19,16 @@
 //TEST_CLASS(MyFixtureName, GenerateTestDirectory)
 TEST_CLASS(MyFixtureName, "Game.Example")
 {
+	//static variables setup in BEFORE_ALL and torn down in AFTER_ALL
+
 	//Member variables shared between tests
 
+	BEFORE_ALL()
+	{
+		//static method
+		//delete if empty
+	}
+
 	BEFORE_EACH()
 	{
 		//delete if empty
@@ -31,6 +39,12 @@ TEST_CLASS(MyFixtureName, "Game.Example")
 		//delete if empty
 	}
 
+	AFTER_ALL()
+	{
+		//static method
+		//delete if empty
+	}
+
 	TEST_METHOD(When_Given_Expect)
 	{
 		ASSERT_THAT(IsTrue(false));
@@ -53,6 +67,14 @@ enum class ECQTestSuppressLogBehavior
 static const FString GenerateTestDirectory = TEXT("");
 static constexpr uint32 DefaultFlags = EAutomationTestFlags::ApplicationContextMask | EAutomationTestFlags::ProductFilter;
 
+template <typename T>
+concept HasBeforeAll = requires(T t) { { T::BeforeAll(FString()) }; };
+
+template <typename T>
+concept HasAfterAll = requires(T t) { { T::AfterAll(FString()) }; };
+
+using BeforeAfterAllFunc = void(*)(const FString&);
+
 template <typename AsserterType>
 struct TBaseTest
 {
@@ -104,6 +126,9 @@ struct TBaseTest
 
 	bool bInitializing{ true };
 
+	BeforeAfterAllFunc BeforeAllFunc = nullptr;
+	BeforeAfterAllFunc AfterAllFunc = nullptr;
+
 	FAutomationTestBase& TestRunner;
 	AsserterType Assert;
 	FTestCommandBuilder TestCommandBuilder;
@@ -113,12 +138,13 @@ template <typename AsserterType>
 struct TTestRunner;
 
 template <typename AsserterType>
-using TTestInstanceGenerator = TUniquePtr<TBaseTest<AsserterType>> (*)(TTestRunner<AsserterType>&);
+using TTestInstanceGenerator = TUniquePtr<TBaseTest<AsserterType>>(*)(TTestRunner<AsserterType>&);
 
 template <typename AsserterType>
 struct TTestRunner : public FAutomationTestBase
 {
 	TTestRunner(FString Name, int32 LineNumber, const char* FileName, FString TestDir, uint32 TestFlags, TTestInstanceGenerator<AsserterType> Factory);
+	~TTestRunner();
 
 	uint32 GetTestFlags() const override;
 	FString GetTestSourceFileName() const override;
@@ -143,6 +169,8 @@ struct TTestRunner : public FAutomationTestBase
 	TArray<FString> TestNames;
 	TMap<FString, int32> TestLineNumbers{};
 	TTestInstanceGenerator<AsserterType> TestInstanceFactory;
+	FDelegateHandle BeforeAllDelegate{};
+	FDelegateHandle AfterAllDelegate{};
 
 protected:
 	FString GetBeautifiedTestName() const override;
@@ -160,7 +188,17 @@ struct TTest : TBaseTest<AsserterType>
 	using DerivedType = Derived;
 
 	TTest()
-		: TBaseTest<AsserterType>(*Derived::TestRunner, Derived::TestRunner->bInitializing) {}
+		: TBaseTest<AsserterType>(*Derived::TestRunner, Derived::TestRunner->bInitializing)
+	{
+		if constexpr (HasBeforeAll<Derived>)
+		{
+			this->BeforeAllFunc = Derived::BeforeAll;
+		}
+		if constexpr (HasAfterAll<Derived>)
+		{
+			this->AfterAllFunc = Derived::AfterAll;
+		}
+	}
 
 	void RunTest(const FString& TestName) override;
 
@@ -185,24 +223,24 @@ struct TTest : TBaseTest<AsserterType>
 };
 
 #if WITH_AUTOMATION_WORKER
-	#define _TEST_CLASS_IMPL(_ClassName, _TestDir, _BaseClass, _AsserterType, _TestFlags)                                                  \
-		struct _ClassName;                                                                                                                 \
-		struct F##_ClassName##_Runner : public TTestRunner<_AsserterType>                                                                  \
-		{                                                                                                                                  \
-			F##_ClassName##_Runner()                                                                                                       \
-				: TTestRunner(#_ClassName, __LINE__, __FILE__, _TestDir, _TestFlags, TTest<_ClassName, _AsserterType>::CreateTestClass) {} \
-		};                                                                                                                                 \
-		F##_ClassName##_Runner _ClassName##_RunnerInstance;                                                                                \
-		struct _ClassName : public _BaseClass<_ClassName, _AsserterType>
-
-	#define TEST_METHOD(_MethodName)                                                                       \
-		FFunctionRegistrar reg##_MethodName{ FString(#_MethodName), &DerivedType::_MethodName, __LINE__ }; \
+#define _TEST_CLASS_IMPL(_ClassName, _TestDir, _BaseClass, _AsserterType, _TestFlags)                                                  \
+	struct _ClassName;                                                                                                                 \
+	struct F##_ClassName##_Runner : public TTestRunner<_AsserterType>                                                                  \
+	{                                                                                                                                  \
+		F##_ClassName##_Runner()                                                                                                       \
+			: TTestRunner(#_ClassName, __LINE__, __FILE__, _TestDir, _TestFlags, TTest<_ClassName, _AsserterType>::CreateTestClass) {} \
+	};                                                                                                                                 \
+	F##_ClassName##_Runner _ClassName##_RunnerInstance;                                                                                \
+	struct _ClassName : public _BaseClass<_ClassName, _AsserterType>
+
+#define TEST_METHOD(_MethodName)                                                                       \
+	FFunctionRegistrar reg##_MethodName{ FString(#_MethodName), &DerivedType::_MethodName, __LINE__ }; \
 	void _MethodName()
 #else
-	#define _TEST_CLASS_IMPL(_ClassName, _TestDir, _BaseClass, _AsserterType, _TestFlags) \
-			struct _ClassName : public _BaseClass<_ClassName, _AsserterType>
+#define _TEST_CLASS_IMPL(_ClassName, _TestDir, _BaseClass, _AsserterType, _TestFlags) \
+	struct _ClassName : public _BaseClass<_ClassName, _AsserterType>
 
-	#define TEST_METHOD(_MethodName) void _MethodName()
+#define TEST_METHOD(_MethodName) void _MethodName()
 #endif
 
 #define TEST_CLASS_WITH_ASSERTS(_ClassName, _TestDir, _AsserterType) _TEST_CLASS_IMPL(_ClassName, _TestDir, TTest, _AsserterType, DefaultFlags)
@@ -222,7 +260,9 @@ struct TTest : TBaseTest<AsserterType>
 #define ASSERT_THAT(_assertion) if (!Assert._assertion) {return;}
 #define ASSERT_FAIL(_Msg) Assert.Fail(_Msg); return;
 
+#define BEFORE_ALL() static void BeforeAll(const FString&)
 #define BEFORE_EACH() virtual void Setup() override
 #define AFTER_EACH() virtual void TearDown() override
+#define AFTER_ALL() static void AfterAll(const FString&)
 
 #include "Impl/CQTest.inl"
diff --git a/Engine/Plugins/Tests/CQTest/Source/CQTest/Public/Impl/CQTest.inl b/Engine/Plugins/Tests/CQTest/Source/CQTest/Public/Impl/CQTest.inl
--- a/Engine/Plugins/Tests/CQTest/Source/CQTest/Public/Impl/CQTest.inl
+++ b/Engine/Plugins/Tests/CQTest/Source/CQTest/Public/Impl/CQTest.inl
@@ -149,9 +149,30 @@ inline TTestRunner<AsserterType>::TTestRunner(FString InName, int32 InLineNumber
 
 	CurrentTestPtr = TestInstanceFactory(*this);
 
+	if (CurrentTestPtr->BeforeAllFunc)
+	{
+		BeforeAllDelegate = FAutomationTestFramework::Get().GetOnEnteringTestSection(GetBeautifiedTestName()).AddStatic(CurrentTestPtr->BeforeAllFunc);
+	}
+	if (CurrentTestPtr->AfterAllFunc)
+	{
+		AfterAllDelegate = FAutomationTestFramework::Get().GetOnLeavingTestSection(GetBeautifiedTestName()).AddStatic(CurrentTestPtr->AfterAllFunc);
+	}
+
 	bInitializing = false;
 }
 
+template <typename AsserterType>
+inline TTestRunner<AsserterType>::~TTestRunner() {
+	if (BeforeAllDelegate.IsValid())
+	{
+		FAutomationTestFramework::Get().GetOnEnteringTestSection(GetBeautifiedTestName()).Remove(BeforeAllDelegate);
+	}
+	if (AfterAllDelegate.IsValid())
+	{
+		FAutomationTestFramework::Get().GetOnLeavingTestSection(GetBeautifiedTestName()).Remove(AfterAllDelegate);
+	}
+}
+
 template <typename AsserterType>
 inline bool TTestRunner<AsserterType>::RunTest(const FString& RequestedTest)
 {
diff --git a/Engine/Plugins/Tests/CQTest/Source/CQTestTests/Private/CQTestTests.cpp b/Engine/Plugins/Tests/CQTest/Source/CQTestTests/Private/CQTestTests.cpp
--- a/Engine/Plugins/Tests/CQTest/Source/CQTestTests/Private/CQTestTests.cpp
+++ b/Engine/Plugins/Tests/CQTest/Source/CQTestTests/Private/CQTestTests.cpp
@@ -181,6 +181,56 @@ namespace CQTestTests
 		}
 	};
 
+	TEST_CLASS(BeforeAndAfterAll, "TestFramework.CQTest.Core")
+	{
+		inline static uint32 BeforeAllCallCount = 0;
+		FString ExpectedError = TEXT("Expected Error Message");
+
+		BEFORE_ALL()
+		{
+			BeforeAllCallCount++;
+		}
+
+		AFTER_ALL()
+		{
+			BeforeAllCallCount = 0;
+		}
+
+		BEFORE_EACH()
+		{
+			ASSERT_THAT(AreEqual(1, BeforeAllCallCount));
+		}
+
+		AFTER_EACH()
+		{
+			ClearExpectedError(*this->TestRunner, ExpectedError);
+			ASSERT_THAT(AreEqual(1, BeforeAllCallCount));
+		}
+
+		TEST_METHOD(StaticMember_IsAvailable_DuringTest)
+		{
+			ASSERT_THAT(AreEqual(1, BeforeAllCallCount));
+		}
+
+		TEST_METHOD(BeforeAll_IsCalled_OnlyOnce)
+		{
+			ASSERT_THAT(AreEqual(1, BeforeAllCallCount));
+		}
+
+		TEST_METHOD(AfterAll_WhenTestFails_StillFires)
+		{
+			Assert.Fail(ExpectedError);
+		}
+	};
+
+	static_assert(HasBeforeAll<BeforeAndAfterAll>);
+	static_assert(HasAfterAll<BeforeAndAfterAll>);
+
+	TEST(ValidateAfterAll, "TestFramework.CQTest.Core")
+	{
+		ASSERT_THAT(AreEqual(0, BeforeAndAfterAll::BeforeAllCallCount));
+	}
+
 	// --------------------------------------------------------
 	// Latent commands are awaited
 	// --------------------------------------------------------
diff --git a/Engine/Plugins/Tests/CQTest/Source/CQTestTests/Private/ExtensionTests.cpp b/Engine/Plugins/Tests/CQTest/Source/CQTestTests/Private/ExtensionTests.cpp
--- a/Engine/Plugins/Tests/CQTest/Source/CQTestTests/Private/ExtensionTests.cpp
+++ b/Engine/Plugins/Tests/CQTest/Source/CQTestTests/Private/ExtensionTests.cpp
@@ -31,17 +31,48 @@ CUSTOM_ASSERT_TEST_CLASS(CustomAsserts, "TestFramework.CQTest.Core")
 template<typename Derived, typename AsserterType>
 struct TCustomBaseClass : public TTest<Derived, AsserterType>
 {
+	inline static uint32 BaseValue = 0;
 	uint32 SpecialValue = 42;
+
+	BEFORE_ALL()
+	{
+		BaseValue = 42;
+	}
+
+	AFTER_ALL()
+	{
+		BaseValue = 0;
+	}
 };
 
 #define CUSTOM_BASE_TEST_CLASS(_ClassName, _TestDir) TEST_CLASS_WITH_BASE(_ClassName, _TestDir, TCustomBaseClass)
 
-CUSTOM_BASE_TEST_CLASS(CustomBase, "TestFramework.CQTest.Core")
+CUSTOM_BASE_TEST_CLASS(DerivedTest, "TestFramework.CQTest.Core")
 {
-	TEST_METHOD(CustomTestClass_WithCustomBase_InheritsFromBaseClass)
+	inline static uint32 DerivedValue = 0;
+
+	BEFORE_ALL()
+	{
+		TCustomBaseClass::BeforeAll(FString());
+		DerivedValue = BaseValue;
+	}
+
+	AFTER_ALL()
+	{
+		DerivedValue = 0;
+		TCustomBaseClass::AfterAll(FString());
+	}
+
+	TEST_METHOD(DerivedTestClass_WithCustomBase_InheritsFromBaseClass)
 	{
 		ASSERT_THAT(AreEqual(42, SpecialValue));
 	}
+
+	TEST_METHOD(DerivedTestClass_WithBeforeAll_CanUseBaseBeforeAll)
+	{
+		ASSERT_THAT(AreEqual(42, BaseValue));
+		ASSERT_THAT(AreEqual(BaseValue, DerivedValue));
+	}
 };
 
 template <typename Derived, typename AsserterType>
