- Improve FPIENetworkComponent and added the ability to have a client join an already in progress network session.
- Added improvements to the NoDiscardAsserter to use explicit methods for FString

#rnx

diff --git a/Engine/Plugins/Tests/CQTest/Source/CQTest/Private/Components/PIENetworkComponent.cpp b/Engine/Plugins/Tests/CQTest/Source/CQTest/Private/Components/PIENetworkComponent.cpp
--- a/Engine/Plugins/Tests/CQTest/Source/CQTest/Private/Components/PIENetworkComponent.cpp
+++ b/Engine/Plugins/Tests/CQTest/Source/CQTest/Private/Components/PIENetworkComponent.cpp
@@ -3,6 +3,7 @@
 #include "Components/PIENetworkComponent.h"
 
 #if ENABLE_PIE_NETWORK_TEST
+
 #include "LevelEditor.h"
 #include "Editor/UnrealEdEngine.h"
 #include "Settings/LevelEditorPlaySettings.h"
@@ -31,8 +32,10 @@ FBasePIENetworkComponent::FBasePIENetworkComponent(FAutomationTestBase* InTestRu
 		->Do(TEXT("Stop PIE"), [this]() { StopPie(); })
 		.Then(TEXT("Create New Map"), [this]() { FAutomationEditorCommonUtils::CreateNewMap(); })
 		.Then(TEXT("Start PIE"), [this]() { StartPie(); })
-		.Until(TEXT("Collect PIE Worlds"), [this]() { return CollectPieWorlds(); })
-		.Until(TEXT("Await Connections"), [this]() { return AwaitConnections(); })
+		.Until(TEXT("Set Worlds"), [this]() { return SetWorlds(); })
+		.Then(TEXT("Setup Packet Settings"), [this]() { SetPacketSettings(); })
+		.Then(TEXT("Connect Clients to Server"), [this]() { ConnectClientsToServer(); })
+		.Until(TEXT("Await Clients Ready"), [this]() { return AwaitClientsReady(); })
 		.OnTearDown(TEXT("Restore Editor State"), [this]() { RestoreState(); });
 }
 
@@ -64,7 +67,8 @@ FBasePIENetworkComponent& FBasePIENetworkComponent::Until(TFunction<bool()> Quer
 	CommandBuilder->Until(Query, Timeout);
 	return *this;
 }
-FBasePIENetworkComponent& FBasePIENetworkComponent::Until(TCHAR* Description, TFunction<bool()> Query, FTimespan Timeout)
+
+FBasePIENetworkComponent& FBasePIENetworkComponent::Until(const TCHAR* Description, TFunction<bool()> Query, FTimespan Timeout)
 {
 	CommandBuilder->Until(Description, Query, Timeout);
 	return *this;
@@ -75,7 +79,8 @@ FBasePIENetworkComponent& FBasePIENetworkComponent::StartWhen(TFunction<bool()>
 	CommandBuilder->StartWhen(Query, Timeout);
 	return *this;
 }
-FBasePIENetworkComponent& FBasePIENetworkComponent::StartWhen(TCHAR* Description, TFunction<bool()> Query, FTimespan Timeout)
+
+FBasePIENetworkComponent& FBasePIENetworkComponent::StartWhen(const TCHAR* Description, TFunction<bool()> Query, FTimespan Timeout)
 {
 	CommandBuilder->StartWhen(Description, Query, Timeout);
 	return *this;
@@ -127,81 +132,125 @@ void FBasePIENetworkComponent::StartPie()
 	GUnrealEd->StartQueuedPlaySessionRequest();
 }
 
-bool FBasePIENetworkComponent::CollectPieWorlds()
+bool FBasePIENetworkComponent::SetWorlds() 
 {
-	const auto& WorldContexts = GEngine->GetWorldContexts();
-	UWorld* ServerWorld = nullptr;
-	TArray<UWorld*> ClientWorlds;
+	auto IsValidContext = [](const FWorldContext& Context) -> bool {
+		return Context.WorldType == EWorldType::PIE && 
+			IsValid(Context.World()) && 
+			IsValid(Context.World()->GetNetDriver());
+	};
 
-	for (const auto& WorldContext : WorldContexts)
+	auto IsValidServerWorld = [this](UWorld* World) -> bool {
+		const bool bIsDedicated = ServerState->bIsDedicatedServer;
+		const bool bExpectsDedicated = World->GetNetMode() == NM_DedicatedServer;
+		return bExpectsDedicated == bIsDedicated;
+	};
+
+	auto IsClientWorldClaimed = [this](UWorld* World) -> bool {
+		for (const auto& State : ClientStates) 
+		{
+			if (IsValid(State->World) && State->World == World) 
+			{
+				return true;
+			}
+		}
+		return false;
+	};
+
+	int32 ClientWorldCount = 0;
+	for (const auto& WorldContext : GEngine->GetWorldContexts())
 	{
-		if (WorldContext.WorldType != EWorldType::PIE || WorldContext.World() == nullptr || WorldContext.World()->GetNetDriver() == nullptr)
+		if (!IsValidContext(WorldContext))
 		{
 			continue;
 		}
+
 		UWorld* World = WorldContext.World();
 		if (World->GetNetDriver()->IsServer())
 		{
-			ServerWorld = World;
-			if (World->GetNetMode() == NM_DedicatedServer && !ServerState->bIsDedicatedServer)
+			if (ServerState->World == nullptr)
 			{
-				TestRunner->AddError(TEXT("Failed to set up dedicated server.  Does your game's editor module override the PIE settings?"));
+				if (!IsValidServerWorld(World)) 
+				{
+					TestRunner->AddError(TEXT("Failed to set up dedicated server. Does your game's editor module override the PIE settings?"));
+					return true;
+				}
 
-				return true;
-			}
-			else if (World->GetNetMode() == NM_ListenServer && ServerState->bIsDedicatedServer)
-			{
-				TestRunner->AddError(TEXT("Failed to set up dedicated server.  Does your game's editor module override the PIE settings?"));
-				return true;
+				ServerState->World = World;
 			}
 		}
 		else
 		{
-			ClientWorlds.Add(World);
+			if (!IsClientWorldClaimed(World))
+			{
+				bool bClaimed = false;
+				for (const auto& State : ClientStates)
+				{
+					if (State->World == nullptr)
+					{
+						State->World = World;
+						bClaimed = true;
+						break;
+					}
+				}
+				if (!bClaimed)
+				{
+					TestRunner->AddError(TEXT("Failed to claim client world. Network component was not able to be initialized."));
+					return true;
+				}
+
+			}
+			ClientWorldCount++;
 		}
 	}
-	if (ServerWorld == nullptr || ClientWorlds.Num() < ServerState->ClientCount)
-	{
-		return false;
-	}
 
+	return IsValid(ServerState->World) && ClientWorldCount == ServerState->ClientCount;
+}
+
+void FBasePIENetworkComponent::SetPacketSettings() const
+{
 	if (PacketSimulationSettings)
 	{
-		ServerWorld->GetNetDriver()->SetPacketSimulationSettings(*PacketSimulationSettings);
-		for (auto& World : ClientWorlds)
+		ServerState->World->GetNetDriver()->SetPacketSimulationSettings(*PacketSimulationSettings);
+		for (const auto& ClientState : ClientStates)
 		{
-			World->GetNetDriver()->SetPacketSimulationSettings(*PacketSimulationSettings);
+			ClientState->World->GetNetDriver()->SetPacketSimulationSettings(*PacketSimulationSettings);
 		}
 	}
+}
 
-	ServerState->World = ServerWorld;
-	for (int32 ClientIndex = 0; ClientIndex < ServerState->ClientCount; ClientIndex++)
+void FBasePIENetworkComponent::ConnectClientsToServer() 
+{
+	auto& ServerConnections = ServerState->World->GetNetDriver()->ClientConnections;
+	for(int32 ClientIndex = ServerState->ClientConnections.Num(); ClientIndex < ServerState->ClientCount; ClientIndex++)
 	{
-		UWorld* ClientWorld = ClientWorlds[ClientIndex];
-		ClientStates[ClientIndex]->World = ClientWorld;
-
-		const int32 ClientLocalPort = ClientWorld->GetNetDriver()->GetLocalAddr()->GetPort();
-		auto ServerConnection = ServerWorld->GetNetDriver()->ClientConnections.FindByPredicate([ClientLocalPort](UNetConnection* ClientConnection) {
+		const int32 ClientLocalPort = ClientStates[ClientIndex]->World->GetNetDriver()->GetLocalAddr()->GetPort();
+		TObjectPtr<UNetConnection>* ServerConnection = ServerConnections.FindByPredicate([ClientLocalPort](UNetConnection* ClientConnection) {
 			return ClientConnection->GetRemoteAddr()->GetPort() == ClientLocalPort;
-			});
-		check(ServerConnection != nullptr);
+		});
+
+		if (ServerConnection == nullptr)
+		{
+			TestRunner->AddError(TEXT("Failed to find connection to server for client. Network component was not able to be initialized."));
+			return;
+		}
+
 		ServerState->ClientConnections[ClientIndex] = *ServerConnection;
 	}
-
-	return true;
 }
 
-bool FBasePIENetworkComponent::AwaitConnections()
+bool FBasePIENetworkComponent::AwaitClientsReady() const
 {
-	if (ServerState == nullptr || ServerState->World == nullptr)
+	if (ServerState == nullptr || !IsValid(ServerState->World))
 	{
-		return true; //Failed to get server state, test will fail
+		TestRunner->AddError(TEXT("Failed to get server state. Network component was not able to be initialized."));
+		return true;
 	}
 	if (ServerState->World->GetNetDriver()->ClientConnections.Num() != ServerState->ClientCount)
 	{
 		return false;
 	}
-	for (UNetConnection* ClientConnection : ServerState->World->GetNetDriver()->ClientConnections)
+	for (const UNetConnection* ClientConnection : ServerState->World->GetNetDriver()->ClientConnections)
 	{
 		if (ClientConnection->ViewTarget == nullptr)
 		{
@@ -220,4 +269,4 @@ void FBasePIENetworkComponent::RestoreState()
 		StateRestorer.Restore();
 	}
 }
-#endif
\ No newline at end of file
+#endif // ENABLE_PIE_NETWORK_TEST
\ No newline at end of file
diff --git a/Engine/Plugins/Tests/CQTest/Source/CQTest/Public/Assert/NoDiscardAsserter.h b/Engine/Plugins/Tests/CQTest/Source/CQTest/Public/Assert/NoDiscardAsserter.h
--- a/Engine/Plugins/Tests/CQTest/Source/CQTest/Public/Assert/NoDiscardAsserter.h
+++ b/Engine/Plugins/Tests/CQTest/Source/CQTest/Public/Assert/NoDiscardAsserter.h
@@ -64,6 +64,18 @@ public:
 	template <typename TExpected, typename TActual>
 	[[nodiscard]] bool AreEqual(const TExpected& Expected, const TActual& Actual, const FString& FailureMessage);
 
+	template <>
+	[[nodiscard]] bool AreEqual(const FString& Expected, const FString& Actual);
+
+	template <>
+	[[nodiscard]] bool AreEqual(const FString& Expected, const FString& Actual, const char* FailureMessage);
+
+	template <>
+	[[nodiscard]] bool AreEqual(const FString& Expected, const FString& Actual, const TCHAR* FailureMessage);
+
+	template<>
+	[[nodiscard]] bool AreEqual(const FString& Expected, const FString& Actual, const FString& FailureMessage);
+
 	[[nodiscard]] bool AreEqualIgnoreCase(const FString& Expected, const FString& Actual);
 
 	[[nodiscard]] bool AreEqualIgnoreCase(const FString& Expected, const FString& Actual, const char* FailureMessage);
@@ -72,6 +84,14 @@ public:
 
 	[[nodiscard]] bool AreEqualIgnoreCase(const FString& Expected, const FString& Actual, const FString& FailureMessage);
 
+	[[nodiscard]] bool AreNotEqualIgnoreCase(const FString& Expected, const FString& Actual);
+
+	[[nodiscard]] bool AreNotEqualIgnoreCase(const FString& Expected, const FString& Actual, const char* FailureMessage);
+
+	[[nodiscard]] bool AreNotEqualIgnoreCase(const FString& Expected, const FString& Actual, const TCHAR* FailureMessage);
+
+	[[nodiscard]] bool AreNotEqualIgnoreCase(const FString& Expected, const FString& Actual, const FString& FailureMessage);
+
 	template <typename TExpected, typename TActual>
 	[[nodiscard]] bool AreNotEqual(const TExpected& Expected, const TActual& Actual);
 
@@ -84,14 +104,6 @@ public:
 	template <typename TExpected, typename TActual>
 	[[nodiscard]] bool AreNotEqual(const TExpected& Expected, const TActual& Actual, const FString& FailureMessage);
 
-	[[nodiscard]] bool AreNotEqualIgnoreCase(const FString& Expected, const FString& Actual);
-
-	[[nodiscard]] bool AreNotEqualIgnoreCase(const FString& Expected, const FString& Actual, const char* FailureMessage);
-
-	[[nodiscard]] bool AreNotEqualIgnoreCase(const FString& Expected, const FString& Actual, const TCHAR* FailureMessage);
-
-	[[nodiscard]] bool AreNotEqualIgnoreCase(const FString& Expected, const FString& Actual, const FString& FailureMessage);
-
 	template <typename TExpected, typename TActual, typename TEpsilon>
 	[[nodiscard]] bool IsNear(TExpected Expected, TActual Actual, TEpsilon Epsilon);
 
diff --git a/Engine/Plugins/Tests/CQTest/Source/CQTest/Public/Assert/NoDiscardAsserter.inl b/Engine/Plugins/Tests/CQTest/Source/CQTest/Public/Assert/NoDiscardAsserter.inl
--- a/Engine/Plugins/Tests/CQTest/Source/CQTest/Public/Assert/NoDiscardAsserter.inl
+++ b/Engine/Plugins/Tests/CQTest/Source/CQTest/Public/Assert/NoDiscardAsserter.inl
@@ -184,6 +184,31 @@ template <typename TExpected, typename TActual>
 	return true;
 }
 
+template <>
+[[nodiscard]] FORCEINLINE bool FNoDiscardAsserter::AreEqual(const FString& Expected, const FString& Actual)
+{
+	FString Message = FString::Printf(TEXT("Expected %s to equal %s"), *Expected, *Actual);
+	return AreEqual(Expected, Actual, Message);
+}
+
+template <>
+[[nodiscard]] FORCEINLINE bool FNoDiscardAsserter::AreEqual(const FString& Expected, const FString& Actual, const char* FailureMessage)
+{
+	return AreEqual(Expected, Actual, FString(FailureMessage));
+}
+
+template <>
+[[nodiscard]] FORCEINLINE bool FNoDiscardAsserter::AreEqual(const FString& Expected, const FString& Actual, const TCHAR* FailureMessage)
+{
+	return AreEqual(Expected, Actual, FString(FailureMessage));
+}
+
+template <>
+[[nodiscard]] FORCEINLINE bool FNoDiscardAsserter::AreEqual(const FString& Expected, const FString& Actual, const FString& FailureMessage)
+{
+	return TestRunner.TestEqual(FailureMessage, *Actual, *Expected); //FAutomationTestBase expects Actual then Expected
+}
+
 template <typename TExpected, typename TActual>
 [[nodiscard]] FORCEINLINE bool FNoDiscardAsserter::AreNotEqual(const TExpected& Expected, const TActual& Actual)
 {
diff --git a/Engine/Plugins/Tests/CQTest/Source/CQTest/Public/Components/PIENetworkComponent.h b/Engine/Plugins/Tests/CQTest/Source/CQTest/Public/Components/PIENetworkComponent.h
--- a/Engine/Plugins/Tests/CQTest/Source/CQTest/Public/Components/PIENetworkComponent.h
+++ b/Engine/Plugins/Tests/CQTest/Source/CQTest/Public/Components/PIENetworkComponent.h
@@ -10,8 +10,10 @@
 #include "CQTest.h"
 #include "PIENetworkTestStateRestorer.h"
 
+#include "Editor.h"
 #include "Engine/NetDriver.h"
 #include "Engine/PackageMapClient.h"
+#include "UnrealEdGlobals.h"
 
 /*
 //Example boiler plate
@@ -77,6 +79,9 @@ struct FBasePIENetworkComponentState
 
 	/** Used by the server to create a PIE session as a dedicated or listen server. */
 	bool bIsDedicatedServer = true;
+
+	/** Used to track spawned and replicated actors across client and server PIE sessions. */
+	TSet<FNetworkGUID> LocallySpawnedActors{};
 };
 
 /**
@@ -158,7 +163,7 @@ public:
 	 *
 	 * @note Will continue to execute until the function provided evaluates to `true` or the duration exceeds the specified timeout.
 	 */
-	FBasePIENetworkComponent& Until(TCHAR* Description, TFunction<bool()> Query, FTimespan Timeout = FTimespan::FromSeconds(10));
+	FBasePIENetworkComponent& Until(const TCHAR* Description, TFunction<bool()> Query, FTimespan Timeout = FTimespan::FromSeconds(10));
 	
 	/**
 	 * Enqueues a function to the latent command manager for execution until completion or timeout.
@@ -183,7 +188,7 @@ public:
 	 *
 	 * @note Will continue to execute until the function provided evaluates to `true` or the duration exceeds the specified timeout.
 	 */
-	FBasePIENetworkComponent& StartWhen(TCHAR* Description, TFunction<bool()> Query, FTimespan Timeout = FTimespan::FromSeconds(10));
+	FBasePIENetworkComponent& StartWhen(const TCHAR* Description, TFunction<bool()> Query, FTimespan Timeout = FTimespan::FromSeconds(10));
 
 protected:
 	/** Prepares the active PIE sessions to be stopped. */
@@ -191,7 +196,7 @@ protected:
 
 	/** Setup settings for a network session and start PIE sessions for both the server and client with the network settings applied. */
 	void StartPie();
-
+	
 	/**
 	 * Fetch all of the worlds that will be used for the networked session.
 	 *
@@ -199,7 +204,13 @@ protected:
 	 *
 	 * @note Method is expected to be used within the `Until` latent command to then wait until the worlds are ready for use.
 	 */
-	bool CollectPieWorlds();
+	bool SetWorlds();
+
+	/** Sets the packet settings used to simulate packet behavior over a network. */
+	void SetPacketSettings() const;
+
+	/** Connects all available clients to the server. */
+	void ConnectClientsToServer();
 
 	/**
 	 * Go through all of the client connections to make sure they are connected and ready.
@@ -208,7 +219,7 @@ protected:
 	 *
 	 * @note Method is expected to be used within the `Until` latent command to then wait until the worlds are ready for use.
 	 */
-	bool AwaitConnections();
+	bool AwaitClientsReady() const;
 
 	/** Restore back to the state prior to test execution. */
 	void RestoreState();
@@ -220,6 +231,7 @@ protected:
 	FPacketSimulationSettings* PacketSimulationSettings = nullptr;
 	TSubclassOf<AGameModeBase> GameMode = TSubclassOf<AGameModeBase>(nullptr);
 	FPIENetworkTestStateRestorer StateRestorer;
+	TMap<FNetworkGUID, int64> SpawnedActors{};
 };
 
 template <typename NetworkDataType>
@@ -369,6 +381,13 @@ public:
 	 */
 	FPIENetworkComponent& UntilClient(const TCHAR* Description, int32 ClientIndex, TFunction<bool(NetworkDataType&)> Query, FTimespan Timeout = FTimespan::FromSeconds(10));
 
+	/**
+	 * Has a new client request a late join and syncs the server and other clients with the new client PIE session.
+	 *
+	 * @return a reference to this
+	 */
+	FPIENetworkComponent& ThenClientJoins();
+	
 	/**
 	 * Spawns an actor on the server and replicates to the clients.
 	 *
@@ -377,8 +396,21 @@ public:
 	template<typename ActorToSpawn, ActorToSpawn* NetworkDataType::*ResultStorage>
 	FPIENetworkComponent& SpawnAndReplicate();
 
+	template<typename ActorToSpawn, ActorToSpawn* NetworkDataType::*ResultStorage>
+	FPIENetworkComponent& SpawnAndReplicate(const FActorSpawnParameters& SpawnParameters);
+
+	template<typename ActorToSpawn, ActorToSpawn* NetworkDataType::* ResultStorage>
+	FPIENetworkComponent& SpawnAndReplicate(TFunction<void(ActorToSpawn&)> BeforeReplicate);
+
+	template<typename ActorToSpawn, ActorToSpawn* NetworkDataType::* ResultStorage>
+	FPIENetworkComponent& SpawnAndReplicate(const FActorSpawnParameters& SpawnParameters, TFunction<void(ActorToSpawn&)> BeforeReplicate);
+
 private:
 	friend class FNetworkComponentBuilder<NetworkDataType>;
+	template<typename ActorToSpawn, ActorToSpawn* NetworkDataType::* ResultStorage>
+	FPIENetworkComponent& SpawnOnServer(const FActorSpawnParameters& SpawnParameters, TFunction<void(ActorToSpawn&)> BeforeReplicate);
+
+	bool ReplicateToClients(NetworkDataType& ClientState);
 };
 
 
diff --git a/Engine/Plugins/Tests/CQTest/Source/CQTest/Public/Impl/PIENetworkComponent.inl b/Engine/Plugins/Tests/CQTest/Source/CQTest/Public/Impl/PIENetworkComponent.inl
--- a/Engine/Plugins/Tests/CQTest/Source/CQTest/Public/Impl/PIENetworkComponent.inl
+++ b/Engine/Plugins/Tests/CQTest/Source/CQTest/Public/Impl/PIENetworkComponent.inl
@@ -6,9 +6,9 @@
 template <typename NetworkDataType>
 inline FPIENetworkComponent<NetworkDataType>& FPIENetworkComponent<NetworkDataType>::ThenServer(TFunction<void(NetworkDataType&)> Action)
 {
-	CommandBuilder->Do([this, Action]() { Action(static_cast<NetworkDataType&>(*ServerState)); });		
-	return *this;
+	return ThenServer(nullptr, Action);
 }
+
 template <typename NetworkDataType>
 inline FPIENetworkComponent<NetworkDataType>& FPIENetworkComponent<NetworkDataType>::ThenServer(const TCHAR* Description, TFunction<void(NetworkDataType&)> Action)
 {
@@ -19,106 +19,140 @@ inline FPIENetworkComponent<NetworkDataType>& FPIENetworkComponent<NetworkDataTy
 template <typename NetworkDataType>
 inline FPIENetworkComponent<NetworkDataType>& FPIENetworkComponent<NetworkDataType>::ThenClients(TFunction<void(NetworkDataType&)> Action)
 {
-	for (int32 Index = 0; Index < ClientStates.Num(); Index++)
-	{
-		CommandBuilder->Do([this, Action, Index]() { Action(static_cast<NetworkDataType&>(*ClientStates[Index])); });
-	}
-	return *this;
+	return ThenClients(nullptr, Action);
 }
+
 template <typename NetworkDataType>
 inline FPIENetworkComponent<NetworkDataType>& FPIENetworkComponent<NetworkDataType>::ThenClients(const TCHAR* Description, TFunction<void(NetworkDataType&)> Action)
 {
-	for (int32 Index = 0; Index < ClientStates.Num(); Index++)
-	{
-		CommandBuilder->Do(Description, [this, Action, Index]() { Action(static_cast<NetworkDataType&>(*ClientStates[Index])); });
-	}
+	// The outer Do is to delay the for-loop until execution time in case a client joins during the test
+	CommandBuilder->Do(Description, [this, Action]() { 
+		for (int32 Index = 0; Index < ClientStates.Num(); Index++)
+		{
+			Action(static_cast<NetworkDataType&>(*ClientStates[Index]));
+		}
+	});
+
 	return *this;
 }
 
 template <typename NetworkDataType>
 inline FPIENetworkComponent<NetworkDataType>& FPIENetworkComponent<NetworkDataType>::ThenClient(int32 ClientIndex, TFunction<void(NetworkDataType&)> Action)
 {
-	if (!ClientStates.IsValidIndex(ClientIndex))
-	{
-		TestRunner->AddError(FString::Printf(TEXT("Invalid client index specified.  Requested Index %d out of %d"), ClientIndex, ClientStates.Num()));
-		return *this;
-	}
-
-	CommandBuilder->Do([this, Action, ClientIndex]() { Action(static_cast<NetworkDataType&>(*ClientStates[ClientIndex])); });
-	return *this;
+	return ThenClient(nullptr, ClientIndex, Action);
 }
+
 template <typename NetworkDataType>
 inline FPIENetworkComponent<NetworkDataType>& FPIENetworkComponent<NetworkDataType>::ThenClient(const TCHAR* Description, int32 ClientIndex, TFunction<void(NetworkDataType&)> Action)
 {
-	if (!ClientStates.IsValidIndex(ClientIndex))
-	{
-		TestRunner->AddError(FString::Printf(TEXT("Invalid client index specified.  Requested Index %d out of %d"), ClientIndex, ClientStates.Num()));
-		return *this;
-	}
-
-	CommandBuilder->Do(Description, [this, Action, ClientIndex]() { Action(static_cast<NetworkDataType&>(*ClientStates[ClientIndex])); });
+	// The outer Do is to delay the for-loop until execution time in case a client joins during the test
+	CommandBuilder->Do(Description, [this, Action, ClientIndex]() { 
+		if (!ClientStates.IsValidIndex(ClientIndex))
+		{
+			TestRunner->AddError(FString::Printf(TEXT("Invalid client index specified. Requested Index: %d MaxIndex: %d"), ClientIndex, ClientStates.Num() - 1));
+			return;
+		}
+		Action(static_cast<NetworkDataType&>(*ClientStates[ClientIndex]));
+	});
+		
 	return *this;
 }
 
 template <typename NetworkDataType>
 inline FPIENetworkComponent<NetworkDataType>& FPIENetworkComponent<NetworkDataType>::UntilServer(TFunction<bool(NetworkDataType&)> Query, FTimespan Timeout)
 {
-	CommandBuilder->Until(
-		[this, Query]() { return Query(static_cast<NetworkDataType&>(*ServerState)); }, Timeout);
-	return *this;
+	return UntilServer(nullptr, Query, Timeout);
 }
+
 template <typename NetworkDataType>
 inline FPIENetworkComponent<NetworkDataType>& FPIENetworkComponent<NetworkDataType>::UntilServer(const TCHAR* Description, TFunction<bool(NetworkDataType&)> Query, FTimespan Timeout)
 {
-	CommandBuilder->Until(
-		Description, [this, Query]() { return Query(static_cast<NetworkDataType&>(*ServerState)); }, Timeout);
+	CommandBuilder->Until(Description, [this, Query]() { 
+		return Query(static_cast<NetworkDataType&>(*ServerState)); 
+	}, Timeout);
+
 	return *this;
 }
 
 template <typename NetworkDataType>
 inline FPIENetworkComponent<NetworkDataType>& FPIENetworkComponent<NetworkDataType>::UntilClients(TFunction<bool(NetworkDataType&)> Query, FTimespan Timeout)
 {
-	for (int32 Index = 0; Index < ClientStates.Num(); Index++)
-	{
-		CommandBuilder->Until(
-			[this, Query, Index]() { return Query(static_cast<NetworkDataType&>(*ClientStates[Index])); }, Timeout);
-	}
-	return *this;
+	return UntilClients(nullptr, Query, Timeout);
 }
+
 template <typename NetworkDataType>
 inline FPIENetworkComponent<NetworkDataType>& FPIENetworkComponent<NetworkDataType>::UntilClients(const TCHAR* Description, TFunction<bool(NetworkDataType&)> Query, FTimespan Timeout)
 {
-	for (int32 Index = 0; Index < ClientStates.Num(); Index++)
-	{
-		CommandBuilder->Until(
-			Description, [this, Query, Index]() { return Query(static_cast<NetworkDataType&>(*ClientStates[Index])); }, Timeout);
-	}
+	// Capture a mutable array of bools (by value) to track which clients have already finished to avoid calling them again
+	TArray<bool> ClientsFinishedTask{};
+	CommandBuilder->Until(Description, [this, Query, Timeout, ClientsFinishedTask]() mutable {
+		// Array with be initially empty and will need to be resized to match our current client count
+		// It is safe to assume that no clients will join in the middle of this action so the resize will occur only once
+		if (ClientsFinishedTask.Num() < ClientStates.Num())
+		{
+			ClientsFinishedTask.SetNumZeroed(ClientStates.Num());
+		}
+
+		bool bIsAllDone = true;
+		for (int32 Index = 0; Index < ClientStates.Num(); Index++) 
+		{
+			if (!ClientsFinishedTask[Index])
+			{
+				if (Query(static_cast<NetworkDataType&>(*ClientStates[Index]))) 
+				{
+					ClientsFinishedTask[Index] = true;
+				}
+				else 
+				{
+					bIsAllDone = false;
+				}
+			}
+		}
+
+		return bIsAllDone;
+	});
+
 	return *this;
 }
 template <typename NetworkDataType>
 inline FPIENetworkComponent<NetworkDataType>& FPIENetworkComponent<NetworkDataType>::UntilClient(int32 ClientIndex, TFunction<bool(NetworkDataType&)> Query, FTimespan Timeout)
 {
-	if (!ClientStates.IsValidIndex(ClientIndex))
-	{
-		TestRunner->AddError(FString::Printf(TEXT("Invalid client index specified.  Requested Index %d out of %d"), ClientIndex, ClientStates.Num()));
-		return *this;
-	}
+	return UntilClient(nullptr, ClientIndex, Query, Timeout);
+}
+
+template <typename NetworkDataType>
+inline FPIENetworkComponent<NetworkDataType>& FPIENetworkComponent<NetworkDataType>::UntilClient(const TCHAR* Description, int32 ClientIndex, TFunction<bool(NetworkDataType&)> Query, FTimespan Timeout)
+{
+	CommandBuilder->Until(Description, [this, Query, ClientIndex]() {
+		if (!ClientStates.IsValidIndex(ClientIndex))
+		{
+			TestRunner->AddError(FString::Printf(TEXT("Invalid client index specified. Requested Index: %d MaxIndex: %d"), ClientIndex, ClientStates.Num() - 1));
+			return true;
+		}
+		return Query(static_cast<NetworkDataType&>(*ClientStates[ClientIndex]));
+	}, Timeout);
 
-	CommandBuilder->Until(
-		[this, Query, ClientIndex]() { return Query(static_cast<NetworkDataType&>(*ClientStates[ClientIndex])); }, Timeout);
 	return *this;
 }
+
 template <typename NetworkDataType>
-inline FPIENetworkComponent<NetworkDataType>& FPIENetworkComponent<NetworkDataType>::UntilClient(const TCHAR* Description, int32 ClientIndex, TFunction<bool(NetworkDataType&)> Query, FTimespan Timeout)
+inline FPIENetworkComponent<NetworkDataType>& FPIENetworkComponent<NetworkDataType>::ThenClientJoins()
 {
-	if (!ClientStates.IsValidIndex(ClientIndex))
-	{
-		TestRunner->AddError(FString::Printf(TEXT("Invalid client index specified.  Requested Index %d out of %d"), ClientIndex, ClientStates.Num()));
-		return *this;
-	}
+	Do(TEXT("Update Server State"), [this]() {
+		int32 NextClientIndex = ServerState->ClientCount++;
+		ClientStates.Add(MakeUnique<NetworkDataType>(NetworkDataType{}));
+		ClientStates.Last()->ClientIndex = NextClientIndex;
+		ServerState->ClientConnections.SetNum(ServerState->ClientCount);
+		GEditor->RequestLateJoin();
+	})
+	.Until(TEXT("Setting Worlds"), [this]() { return SetWorlds(); })
+	.Then(TEXT("Setup Packet Settings"), [this]() { SetPacketSettings(); })
+	.Then(TEXT("Connect Clients to Server"), [this]() { ConnectClientsToServer(); });
 
-	CommandBuilder->Until(
-		Description, [this, Query, ClientIndex]() { return Query(static_cast<NetworkDataType&>(*ClientStates[ClientIndex])); }, Timeout);
+	UntilClient(TEXT("Replicate to new Client"), ServerState->ClientCount, [this](NetworkDataType& State) {
+		return ReplicateToClients(State);
+	});
+	
 	return *this;
 }
 
@@ -126,36 +160,102 @@ template <typename NetworkDataType>
 template <typename ActorToSpawn, ActorToSpawn* NetworkDataType::*ResultStorage>
 inline FPIENetworkComponent<NetworkDataType>& FPIENetworkComponent<NetworkDataType>::SpawnAndReplicate()
 {
+	return SpawnAndReplicate<ActorToSpawn, ResultStorage>({}, {});
+}
+
+template <typename NetworkDataType>
+template<typename ActorToSpawn, ActorToSpawn* NetworkDataType::* ResultStorage>
+inline FPIENetworkComponent<NetworkDataType>& FPIENetworkComponent<NetworkDataType>::SpawnAndReplicate(const FActorSpawnParameters& SpawnParameters) {
+	return SpawnAndReplicate<ActorToSpawn, ResultStorage>(SpawnParameters, {});
+}
+
+template <typename NetworkDataType>
+template<typename ActorToSpawn, ActorToSpawn* NetworkDataType::* ResultStorage>
+inline FPIENetworkComponent<NetworkDataType>& FPIENetworkComponent<NetworkDataType>::SpawnAndReplicate(TFunction<void(ActorToSpawn&)> BeforeReplicate) {
+	return SpawnAndReplicate<ActorToSpawn, ResultStorage>({}, BeforeReplicate);
+}
+
+template <typename NetworkDataType>
+template<typename ActorToSpawn, ActorToSpawn* NetworkDataType::* ResultStorage>
+inline FPIENetworkComponent<NetworkDataType>& FPIENetworkComponent<NetworkDataType>::SpawnAndReplicate(const FActorSpawnParameters& SpawnParameters, TFunction<void(ActorToSpawn&)> BeforeReplicate) 
+{
+	SpawnOnServer<ActorToSpawn, ResultStorage>(SpawnParameters, BeforeReplicate);
+	UntilClients([this](NetworkDataType& ClientState) { return ReplicateToClients(ClientState); });
+		
+	return *this;
+}
+
+template <typename NetworkDataType>
+template<typename ActorToSpawn, ActorToSpawn* NetworkDataType::* ResultStorage>
+inline FPIENetworkComponent<NetworkDataType>& FPIENetworkComponent<NetworkDataType>::SpawnOnServer(const FActorSpawnParameters& SpawnParameters, TFunction<void(ActorToSpawn&)> BeforeReplicate) {
 	static_assert(std::is_convertible_v<ActorToSpawn*, AActor*>, "ActorToSpawn must derive from AActor");
+	static_assert(std::is_default_constructible<NetworkDataType>::value, "NetworkDataType must have a default constructor accessible");
 
-	TSharedPtr<ActorToSpawn*> SharedStorage = MakeShareable(new ActorToSpawn*(nullptr));
-	TSharedRef<FNetworkGUID> SharedGuid = MakeShareable(new FNetworkGUID());
+	// TODO: Consider passing in a constructed actor from an FTestSpawner instead
+	// That would allow using TObjectBuilder as well.
+	// Need a version which takes the Server's world instead of creating its own
 
-	ThenServer(TEXT("Spawning Actor On Server"), [SharedStorage](NetworkDataType& State) {
-		*SharedStorage = State.World->template SpawnActor<ActorToSpawn>();
-		if (ResultStorage != nullptr)
+	TSharedPtr<ActorToSpawn*> ServerActor = MakeShareable(new ActorToSpawn * (nullptr));
+
+	ThenServer(TEXT("Spawning Actor On Server"), [ServerActor, SpawnParameters, BeforeReplicate](NetworkDataType& State) {
+		*ServerActor = State.World->template SpawnActor<ActorToSpawn>(FVector::ZeroVector, FRotator::ZeroRotator, SpawnParameters);
+		if(BeforeReplicate)
 		{
-			State.*ResultStorage = *SharedStorage;
+			BeforeReplicate(**ServerActor);
 		}
-	}).UntilServer(TEXT("Waiting for NetGUID"), [SharedStorage, SharedGuid](NetworkDataType& State) {
-		*SharedGuid = State.World->GetNetDriver()->GuidCache->GetNetGUID(*SharedStorage);
-		return SharedGuid->IsValid();
-	}).UntilClients(TEXT("Waiting for Replication on Clients"), [SharedGuid](NetworkDataType& State) {
-		ActorToSpawn* ClientActor = Cast<ActorToSpawn>(State.World->GetNetDriver()->GuidCache->GetObjectFromNetGUID(*SharedGuid, true));
-		if (ClientActor == nullptr)
+		if (ResultStorage != nullptr)
 		{
-			return false;
+			State.*ResultStorage = *ServerActor;
 		}
-		if (ResultStorage != nullptr)
+	})
+	.UntilServer(TEXT("Waiting for NetGUID"), [this, ServerActor](NetworkDataType& State) {
+		FNetworkGUID NetGUID = State.World->GetNetDriver()->GuidCache->GetNetGUID(*ServerActor);
+		if (!NetGUID.IsValid())
 		{
-			State.*ResultStorage = ClientActor;
+			return false;
 		}
+
+		// Calculate the pointer offset to the storage location on NetworkDataType
+		// Do this by allocating a temporary NetworkDataType object and then calculate
+		// the address of the storage location as an int64
+		// This allows ReplicateToClients to not need the ActorToSpawn or ResultStorage template parameters
+		// which in turn allows ThenClientJoins to use ReplicateToClients
+		NetworkDataType TempDataType {};
+		int64 StorageOffset = reinterpret_cast<int64>(&(TempDataType.*ResultStorage)) - reinterpret_cast<int64>(&TempDataType);
+
+		SpawnedActors.Add(NetGUID, StorageOffset);
+		State.LocallySpawnedActors.Add(NetGUID);
 		return true;
 	});
-
+	
 	return *this;
 }
 
+template <typename NetworkDataType>
+inline bool FPIENetworkComponent<NetworkDataType>::ReplicateToClients(NetworkDataType& ClientState)
+{
+	for(const auto& [NetGUID, StorageOffset] : SpawnedActors)
+	{
+		if(ClientState.LocallySpawnedActors.Contains(NetGUID))
+		{
+			continue;
+		}
+
+		AActor* ClientActor = Cast<AActor>(ClientState.World->GetNetDriver()->GuidCache->GetObjectFromNetGUID(NetGUID, true));
+		if (ClientActor != nullptr)
+		{
+			// The other half of the offset implementation.
+			// Interpret the address of the State object as a 1-byte aligned pointer, and add the offset
+			// Then interpret that as a pointer to a pointer to an AActor which we can assign
+			AActor** Storage = reinterpret_cast<AActor**>(reinterpret_cast<char*>(&ClientState) + StorageOffset);
+			*Storage = ClientActor;
+			ClientState.LocallySpawnedActors.Add(NetGUID);
+		}
+	}
+
+	return ClientState.LocallySpawnedActors.Num() == SpawnedActors.Num();
+}
+
 ///////////////////////////////////////////////////////////////////////
 // FNetworkComponentBuilder
 
diff --git a/Engine/Plugins/Tests/CQTest/Source/CQTestTests/Private/Assert/AssertionTests.cpp b/Engine/Plugins/Tests/CQTest/Source/CQTestTests/Private/Assert/AssertionTests.cpp
--- a/Engine/Plugins/Tests/CQTest/Source/CQTestTests/Private/Assert/AssertionTests.cpp
+++ b/Engine/Plugins/Tests/CQTest/Source/CQTestTests/Private/Assert/AssertionTests.cpp
@@ -12,7 +12,7 @@ namespace CQTests
 
 	TEST_CLASS(NoDiscardAssert_Errors, "TestFramework.CQTest.Core")
 	{
-
+		FNoDiscardAsserter Asserter{ *TestRunner };
 		TEST_METHOD(AssertFail_WithMessage_AddsError)
 		{
 			Assert.ExpectError(AnyError);
@@ -30,55 +30,55 @@ namespace CQTests
 
 		TEST_METHOD(Assertions_Accept_RawStrings)
 		{
-			Assert.ExpectError("Hello World");
-			Assert.Fail("Hello World");
+			Asserter.ExpectError("Hello World");
+			Asserter.Fail("Hello World");
 		}
 
 		TEST_METHOD(Assertions_Accept_TCharArrays)
 		{
-			Assert.ExpectError(TEXT("Hello World"));
-			Assert.Fail(TEXT("Hello World"));
+			Asserter.ExpectError(TEXT("Hello World"));
+			Asserter.Fail(TEXT("Hello World"));
 		}
 
 		TEST_METHOD(Assertions_Accept_FStrings)
 		{
 			FString message = TEXT("Hello World");
-			Assert.ExpectError(message);
-			Assert.Fail(TEXT("Hello WorldA"));
+			Asserter.ExpectError(message);
+			Asserter.Fail(message);
 		}
 
 		TEST_METHOD(AssertExpectError_WithMultipleErrors_Succeeds)
 		{
-			Assert.ExpectError("", 3);
-			Assert.Fail("One");
-			Assert.Fail("Two");
-			Assert.Fail("Three");
+			Asserter.ExpectError(AnyError, 3);
+			Asserter.Fail("One");
+			Asserter.Fail("Two");
+			Asserter.Fail("Three");
 		}
 
 		TEST_METHOD(AsserterExpectError_WithZeroExpected_AcceptsAnyNumber)
 		{
-			Assert.ExpectError("", 0);
-			Assert.Fail("One");
-			Assert.Fail("Two");
-			Assert.Fail("Three");
+			Asserter.ExpectError(AnyError, 0);
+			Asserter.Fail("One");
+			Asserter.Fail("Two");
+			Asserter.Fail("Three");
 		}
 
 		TEST_METHOD(AssertExpectError_WithMatchingError_Succeeds)
 		{
-			Assert.ExpectError("Hello World");
-			Assert.Fail("Hello World");
+			Asserter.ExpectError("Hello World");
+			Asserter.Fail("Hello World");
 		}
 
 		TEST_METHOD(AssertExpectError_WithRegexSymbols_EscapesRegex)
 		{
-			Assert.ExpectError("[^abc]");
-			Assert.Fail("[^abc]");
+			Asserter.ExpectError("[^abc]");
+			Asserter.Fail("[^abc]");
 		}
 
 		TEST_METHOD(AssertExpectErrorRegex_WithRegex_Succeeds)
 		{
-			Assert.ExpectErrorRegex("\\w+");
-			Assert.Fail("abc");
+			Asserter.ExpectErrorRegex("\\w+");
+			Asserter.Fail("abc");
 		}
 	};
 
diff --git a/Engine/Plugins/Tests/CQTest/Source/CQTestTests/Private/Components/CQTestGameInstance.h b/Engine/Plugins/Tests/CQTest/Source/CQTestTests/Private/Components/CQTestGameInstance.h
new file mode 100644
--- /dev/null
+++ b/Engine/Plugins/Tests/CQTest/Source/CQTestTests/Private/Components/CQTestGameInstance.h
@@ -0,0 +1,17 @@
+// Copyright Epic Games, Inc. All Rights Reserved.
+
+#pragma once
+
+#include "Engine/GameInstance.h"
+
+#include "CQTestGameInstance.generated.h"
+
+/** Basic GameInstance with a test value. */
+UCLASS()
+class UCQGameInstanceClass : public UGameInstance
+{
+	GENERATED_BODY()
+
+public:
+	int TestValue{ 42 };
+};
diff --git a/Engine/Plugins/Tests/CQTest/Source/CQTestTests/Private/Components/CQTestGameMode.h b/Engine/Plugins/Tests/CQTest/Source/CQTestTests/Private/Components/CQTestGameMode.h
new file mode 100644
--- /dev/null
+++ b/Engine/Plugins/Tests/CQTest/Source/CQTestTests/Private/Components/CQTestGameMode.h
@@ -0,0 +1,16 @@
+// Copyright Epic Games, Inc. All Rights Reserved.
+
+#pragma once
+
+#include "GameFramework/GameModeBase.h"
+#include "CQTestGameMode.generated.h"
+
+/** Basic GameMode with a test value. */
+UCLASS()
+class ACQTestGameMode : public AGameModeBase
+{
+	GENERATED_BODY()
+
+public:
+	int32 TestValue{ 42 };
+};
\ No newline at end of file
diff --git a/Engine/Plugins/Tests/CQTest/Source/CQTestTests/Private/Components/PIENetworkComponentTests.cpp b/Engine/Plugins/Tests/CQTest/Source/CQTestTests/Private/Components/PIENetworkComponentTests.cpp
--- a/Engine/Plugins/Tests/CQTest/Source/CQTestTests/Private/Components/PIENetworkComponentTests.cpp
+++ b/Engine/Plugins/Tests/CQTest/Source/CQTestTests/Private/Components/PIENetworkComponentTests.cpp
@@ -2,8 +2,9 @@
 
 #include "CQTest.h"
 #include "Components/PIENetworkComponent.h"
+#include "CQTestGameInstance.h"
+#include "CQTestGameMode.h"
 #include "CQTestUnitTestHelper.h"
-
 #include "TestReplicatedActor.h"
 
 #include "Kismet/GameplayStatics.h"
@@ -35,52 +36,52 @@ NETWORK_TEST_CLASS(StateTest, "TestFramework.CQTest.Network")
 	TEST_METHOD(Network_WithMultipleSteps_TriggersStepsInOrder)
 	{
 		Network
-			.ThenServer([&](DerivedState&) { ASSERT_THAT(AreEqual(0, SharedNumber++)); })
-			.ThenClient(0, [&](DerivedState&) { ASSERT_THAT(AreEqual(1, SharedNumber++)); })
-			.ThenServer([&](DerivedState&) { ASSERT_THAT(AreEqual(2, SharedNumber++)); })
-			.Then([&]() { ASSERT_THAT(AreEqual(3, SharedNumber++)); });
+			.ThenServer([this](DerivedState&) { ASSERT_THAT(AreEqual(0, SharedNumber++)); })
+			.ThenClient(0, [this](DerivedState&) { ASSERT_THAT(AreEqual(1, SharedNumber++)); })
+			.ThenServer([this](DerivedState&) { ASSERT_THAT(AreEqual(2, SharedNumber++)); })
+			.Then([this]() { ASSERT_THAT(AreEqual(3, SharedNumber++)); });
 	}
 
 	TEST_METHOD(Network_WithServerCommands_RetainsStateBetweenCalls)
 	{
 		Network
-			.ThenServer([&](DerivedState& State) { ASSERT_THAT(AreEqual(0, State.IndependentNumber++)); })
-			.ThenServer([&](DerivedState& State) { ASSERT_THAT(AreEqual(1, State.IndependentNumber++)); });
+			.ThenServer([this](DerivedState& State) { ASSERT_THAT(AreEqual(0, State.IndependentNumber++)); })
+			.ThenServer([this](DerivedState& State) { ASSERT_THAT(AreEqual(1, State.IndependentNumber++)); });
 	}
 
 	TEST_METHOD(Network_WithClientCommands_RetainsStateBetweenCalls)
 	{
 		Network
-			.ThenClient(0, [&](DerivedState& State) { ASSERT_THAT(AreEqual(0, State.IndependentNumber++)); })
-			.ThenClient(0, [&](DerivedState& State) { ASSERT_THAT(AreEqual(1, State.IndependentNumber++)); });
+			.ThenClient(0, [this](DerivedState& State) { ASSERT_THAT(AreEqual(0, State.IndependentNumber++)); })
+			.ThenClient(0, [this](DerivedState& State) { ASSERT_THAT(AreEqual(1, State.IndependentNumber++)); });
 	}
 
 	TEST_METHOD(Network_WithClientAndServerCommands_DoNotShareState)
 	{
 		Network
-			.ThenServer([&](DerivedState& State) { State.IndependentNumber++; })
-			.ThenClient(0, [&](DerivedState& State) { ASSERT_THAT(AreEqual(0, State.IndependentNumber)); });
+			.ThenServer([this](DerivedState& State) { State.IndependentNumber++; })
+			.ThenClient(0, [this](DerivedState& State) { ASSERT_THAT(AreEqual(0, State.IndependentNumber)); });
 	}
 
 	TEST_METHOD(Network_WithMultipleClients_DoNotShareState)
 	{
 		Network
-			.ThenClients([&](DerivedState& State) { State.IndependentNumber = State.ClientIndex; })
-			.ThenClients([&](DerivedState& State) { ASSERT_THAT(AreEqual(State.ClientIndex, State.IndependentNumber)); });
+			.ThenClients([this](DerivedState& State) { State.IndependentNumber = State.ClientIndex; })
+			.ThenClients([this](DerivedState& State) { ASSERT_THAT(AreEqual(State.ClientIndex, State.IndependentNumber)); });
 	}
 
 	TEST_METHOD(Network_WithTickingServerCommand_TicksUntilDone)
 	{
 		Network
-			.UntilServer([&](DerivedState& State) { return ++State.IndependentNumber > 4; })
-			.ThenServer([&](DerivedState& State) { ASSERT_THAT(AreEqual(State.IndependentNumber, 5)); });
+			.UntilServer([this](DerivedState& State) { return ++State.IndependentNumber > 4; })
+			.ThenServer([this](DerivedState& State) { ASSERT_THAT(AreEqual(State.IndependentNumber, 5)); });
 	}
 
 	TEST_METHOD(Network_WithTickingClientCommands_TicksEachCommand)
 	{
 		Network
-			.UntilClients([&](DerivedState& State) { SharedNumber++;  return ++State.IndependentNumber > 4; })
-			.Then([&]() { ASSERT_THAT(AreEqual(15, SharedNumber)); });
+			.UntilClients([this](DerivedState& State) { SharedNumber++;  return ++State.IndependentNumber > 4; })
+			.Then([this]() { ASSERT_THAT(AreEqual(15, SharedNumber)); });
 	}
 };
 
@@ -91,8 +92,11 @@ NETWORK_TEST_CLASS(ReplicationTest, "TestFramework.CQTest.Network")
 		ATestReplicatedActor* ReplicatedActor = nullptr;
 	};
 
+	const int32 ExpectedReplicatedValue = 42;
+
 	FPIENetworkComponent<DerivedState> Network{ TestRunner, TestCommandBuilder, bInitializing };
-	BEFORE_EACH() {
+	BEFORE_EACH() 
+	{
 		FNetworkComponentBuilder<DerivedState>()
 			.WithGameInstanceClass(UGameInstance::StaticClass())
 			.WithGameMode(AGameModeBase::StaticClass())
@@ -102,10 +106,10 @@ NETWORK_TEST_CLASS(ReplicationTest, "TestFramework.CQTest.Network")
 	TEST_METHOD(SpawnAndReplicateActor_WithReplicatedActor_ProvidesActorToClients)
 	{
 		Network.SpawnAndReplicate<ATestReplicatedActor, &DerivedState::ReplicatedActor>()
-			.ThenServer([&](DerivedState& ServerState) {
+			.ThenServer([this](DerivedState& ServerState) {
 				ASSERT_THAT(IsNotNull(ServerState.ReplicatedActor));
 			})
-			.ThenClients([&](DerivedState& ClientState) {
+			.ThenClients([this](DerivedState& ClientState) {
 				ASSERT_THAT(IsNotNull(ClientState.ReplicatedActor));
 			});
 	}
@@ -113,44 +117,179 @@ NETWORK_TEST_CLASS(ReplicationTest, "TestFramework.CQTest.Network")
 	TEST_METHOD(SpawnAndReplicateActor_ThenUpdateProperty_UpdatesPropertyOnClients)
 	{
 		Network.SpawnAndReplicate<ATestReplicatedActor, &DerivedState::ReplicatedActor>()
-			.ThenServer(TEXT("Server Set Value"), [](DerivedState& ServerState) {
-				ServerState.ReplicatedActor->ReplicatedInt = 42;
+			.ThenServer(TEXT("Server Set Value"), [this](DerivedState& ServerState) {
+				ServerState.ReplicatedActor->ReplicatedInt = ExpectedReplicatedValue;
 			})
-			.UntilClients(TEXT("Clients Check Value"), [](DerivedState& ClientState) {
-				return ClientState.ReplicatedActor->ReplicatedInt == 42;
+			.UntilClients(TEXT("Clients Check Value"), [this](DerivedState& ClientState) {
+				return ClientState.ReplicatedActor->ReplicatedInt == ExpectedReplicatedValue;
+			});
+	}
+
+	TEST_METHOD(SpawnAndReplicateActor_WithSpawnParameters_PassesParametersToSpawnedObject)
+	{
+		FActorSpawnParameters SpawnParameters;
+		SpawnParameters.SpawnCollisionHandlingOverride = ESpawnActorCollisionHandlingMethod::AlwaysSpawn;
+		Network.SpawnAndReplicate<ATestReplicatedActor, &DerivedState::ReplicatedActor>(SpawnParameters)
+			.UntilClients([SpawnParameters](DerivedState& ClientState) {
+				return SpawnParameters.SpawnCollisionHandlingOverride == ClientState.ReplicatedActor->SpawnCollisionHandlingMethod;
+			});
+	}
+
+	TEST_METHOD(SpawnAndReplicateActor_WithBeforeReplicates_InvokesBeforeReplicate) {
+		auto BeforeReplicate = [this](ATestReplicatedActor& Actor) {
+			Actor.ReplicatedInt = ExpectedReplicatedValue;
+		};
+
+		Network.SpawnAndReplicate<ATestReplicatedActor, &DerivedState::ReplicatedActor>(BeforeReplicate)
+			.UntilClients([this](DerivedState& ClientState) {
+				return ClientState.ReplicatedActor->ReplicatedInt == ExpectedReplicatedValue;
+			});
+	}
+
+	TEST_METHOD(SpawnAndReplicateActor_WithSpawnParametersAndBeforeReplicate_UsesBoth)
+	{
+		auto BeforeReplicate = [this](ATestReplicatedActor& Actor) {
+			Actor.ReplicatedInt = ExpectedReplicatedValue;
+		};
+
+		FActorSpawnParameters SpawnParameters;
+		SpawnParameters.SpawnCollisionHandlingOverride = ESpawnActorCollisionHandlingMethod::AlwaysSpawn;
+		Network.SpawnAndReplicate<ATestReplicatedActor, &DerivedState::ReplicatedActor>(SpawnParameters, BeforeReplicate)
+			.UntilClients([this, SpawnParameters](DerivedState& ClientState) {
+				ATestReplicatedActor& actor = *ClientState.ReplicatedActor;
+				return SpawnParameters.SpawnCollisionHandlingOverride == actor.SpawnCollisionHandlingMethod && actor.ReplicatedInt == ExpectedReplicatedValue;
 			});
 	}
 };
 
-NETWORK_TEST_CLASS(GameInstanceTest, "TestFramework.CQTest.Network")
+NETWORK_TEST_CLASS(MultipleActorStateReplication, "TestFramework.CQTest.Network")
 {
-	struct SomeGameInstanceClass : UGameInstance
+	struct DerivedState : public FBasePIENetworkComponentState
 	{
+		ATestReplicatedActor* ReplicatedActor1 = nullptr;
+		ATestReplicatedActor* ReplicatedActor2 = nullptr;
 	};
 
+	const int32 Actor1ExpectedReplicatedValue = 42;
+	const int32 Actor2ExpectedReplicatedValue = 24;
+
+	FPIENetworkComponent<DerivedState> Network{ TestRunner, TestCommandBuilder, bInitializing };
+	BEFORE_EACH()
+	{
+		FNetworkComponentBuilder<DerivedState>()
+			.WithClients(1)
+			.Build(Network);
+	}
+
+	auto MakeSetInt(int32 Value) 
+	{
+		return [Value](ATestReplicatedActor& Actor) {
+			Actor.ReplicatedInt = Value;
+		};
+	}
+
+	TEST_METHOD(SpawnAndReplicateActor_WithMultipleActors_ReplicatesBoth)
+	{
+		Network.SpawnAndReplicate<ATestReplicatedActor, &DerivedState::ReplicatedActor1>(MakeSetInt(Actor1ExpectedReplicatedValue))
+			.SpawnAndReplicate<ATestReplicatedActor, &DerivedState::ReplicatedActor2>(MakeSetInt(Actor2ExpectedReplicatedValue))
+			.ThenServer([this](DerivedState& ServerState) {
+				ASSERT_THAT(AreEqual(Actor1ExpectedReplicatedValue, ServerState.ReplicatedActor1->ReplicatedInt));
+				ASSERT_THAT(AreEqual(Actor2ExpectedReplicatedValue, ServerState.ReplicatedActor2->ReplicatedInt));
+			})
+			.ThenClients([this](DerivedState& ClientState) {
+				ASSERT_THAT(AreEqual(Actor1ExpectedReplicatedValue, ClientState.ReplicatedActor1->ReplicatedInt));
+				ASSERT_THAT(AreEqual(Actor2ExpectedReplicatedValue, ClientState.ReplicatedActor2->ReplicatedInt));
+			});
+	}
+
+	TEST_METHOD(SpawnAndReplicate_WithLateJoin_ReplicatesBoth)
+	{
+		Network.SpawnAndReplicate<ATestReplicatedActor, &DerivedState::ReplicatedActor1>(MakeSetInt(42))
+			.SpawnAndReplicate<ATestReplicatedActor, &DerivedState::ReplicatedActor2>(MakeSetInt(24))
+			.ThenServer([this](DerivedState& ServerState) {
+				ASSERT_THAT(AreEqual(Actor1ExpectedReplicatedValue, ServerState.ReplicatedActor1->ReplicatedInt));
+				ASSERT_THAT(AreEqual(Actor2ExpectedReplicatedValue, ServerState.ReplicatedActor2->ReplicatedInt));
+			})
+			.ThenClientJoins()
+			.ThenClients([this](DerivedState& ClientState) {
+				ASSERT_THAT(AreEqual(Actor1ExpectedReplicatedValue, ClientState.ReplicatedActor1->ReplicatedInt));
+				ASSERT_THAT(AreEqual(Actor2ExpectedReplicatedValue, ClientState.ReplicatedActor2->ReplicatedInt));
+			});
+	}
+};
+
+NETWORK_TEST_CLASS(LateJoinTest, "TestFramework.CQTest.Network")
+{
+	struct DerivedState : public FBasePIENetworkComponentState
+	{
+		ATestReplicatedActor* ReplicatedActor = nullptr;
+	};
+
+	FPIENetworkComponent<DerivedState> Network{ TestRunner, TestCommandBuilder, bInitializing };
+	BEFORE_EACH()
+	{
+		FNetworkComponentBuilder<DerivedState>()
+			.WithGameInstanceClass(UGameInstance::StaticClass())
+			.WithGameMode(AGameModeBase::StaticClass())
+			.WithClients(1)
+			.Build(Network);
+	}
+
+	TEST_METHOD(ThenClientJoins_AfterStart_AddsClient) 
+	{
+		Network.ThenServer([this](DerivedState& ServerState) {
+			ASSERT_THAT(AreEqual(ServerState.ClientCount, 1));
+		})
+		.ThenClientJoins()
+		.ThenServer([this](DerivedState& ServerState) {
+			ASSERT_THAT(AreEqual(ServerState.ClientCount, 2));
+			ASSERT_THAT(AreEqual(ServerState.ClientConnections.Num(), 2));
+		});
+	}
+
+	TEST_METHOD(ThenClientJoins_AfterStart_ReplicatesState) 
+	{
+		Network.SpawnAndReplicate<ATestReplicatedActor, &DerivedState::ReplicatedActor>()
+			.ThenServer([this](DerivedState& ServerState) {
+				ASSERT_THAT(IsNotNull(ServerState.ReplicatedActor));
+			})
+			.ThenClient(0, [this](DerivedState& ClientState) {
+				ASSERT_THAT(IsNotNull(ClientState.ReplicatedActor));
+			})
+			.ThenClientJoins()
+			.ThenClient(1, [this](DerivedState& ClientState) {
+				ASSERT_THAT(IsNotNull(ClientState.ReplicatedActor));
+			});
+	}
+};
+
+NETWORK_TEST_CLASS(GameInstanceTest, "TestFramework.CQTest.Network")
+{
+	const int32 ExpectedReplicatedValue = 42;
+
 	FPIENetworkComponent<FBasePIENetworkComponentState> Network{ TestRunner, TestCommandBuilder, bInitializing };
+
 	BEFORE_EACH()
 	{
 		FNetworkComponentBuilder()
-			.WithGameInstanceClass(FSoftClassPath(SomeGameInstanceClass::StaticClass()))
+			.WithGameInstanceClass(FSoftClassPath(UCQGameInstanceClass::StaticClass()))
 			.WithGameMode(AGameModeBase::StaticClass())
 			.Build(Network);
 	}
 
 	TEST_METHOD(NetworkComponent_WithGameInstanceClass_BuildsNetworkWithProvidedGameInstance)
 	{
-		Network.ThenServer([&](FBasePIENetworkComponentState& ServerState) {
-			auto* GameInstance = Cast<SomeGameInstanceClass>(ServerState.World->GetGameInstance());
+		Network.ThenServer([this](FBasePIENetworkComponentState& ServerState) {
+			UCQGameInstanceClass* GameInstance = Cast<UCQGameInstanceClass>(ServerState.World->GetGameInstance());
 			ASSERT_THAT(IsNotNull(GameInstance));
+			ASSERT_THAT(AreEqual(ExpectedReplicatedValue, GameInstance->TestValue));
 		});
 	}
 };
 
 NETWORK_TEST_CLASS(GameModeTest, "TestFramework.CQTest.Network")
 {
-	struct ATestGameMode : AGameModeBase
-	{
-	};
+	const int32 ExpectedReplicatedValue = 42;
 
 	FPIENetworkComponent<> Network{ TestRunner, TestCommandBuilder, bInitializing };
 
@@ -158,15 +297,16 @@ NETWORK_TEST_CLASS(GameModeTest, "TestFramework.CQTest.Network")
 	{
 		FNetworkComponentBuilder()
 			.WithGameInstanceClass(UGameInstance::StaticClass())
-			.WithGameMode(ATestGameMode::StaticClass())
+			.WithGameMode(ACQTestGameMode::StaticClass())
 			.Build(Network);
 	}
 
 	TEST_METHOD(NetworkComponent_WithGameMode_BuildsNetworkWithProvidedGameMode)
 	{
-		Network.ThenServer([&](auto& ServerState) {
-			auto* GameMode = Cast<ATestGameMode>(ServerState.World->GetAuthGameMode());
+		Network.ThenServer([this](FBasePIENetworkComponentState& ServerState) {
+			ACQTestGameMode* GameMode = Cast<ACQTestGameMode>(ServerState.World->GetAuthGameMode());
 			ASSERT_THAT(IsNotNull(GameMode));
+			ASSERT_THAT(AreEqual(ExpectedReplicatedValue, GameMode->TestValue));
 		});
 	}
 };
@@ -179,9 +319,10 @@ NETWORK_TEST_CLASS(SetupErrorTest, "TestFramework.CQTest.Network")
 	{
 		ClearExpectedError(*TestRunner, TEXT("Failed to initialize Network Component"));
 	}
+
 	TEST_METHOD(NetworkComponent_WithoutUsingBuilder_AddsErrorAndDoesNotCrash)
 	{
-		Network.ThenServer([&](auto&) { Assert.Fail("Unexpected Error"); });
+		Network.ThenServer([this](FBasePIENetworkComponentState&) { Assert.Fail("Unexpected Error"); });
 	}
 };
 
