Fixing crash caused by ending the PIE session externally by adding a callback to the MapTestSpawner when the session ends as the Map may be destructed

#rnx

diff --git a/Engine/Plugins/Tests/CQTest/Source/CQTest/Private/Components/MapTestSpawner.cpp b/Engine/Plugins/Tests/CQTest/Source/CQTest/Private/Components/MapTestSpawner.cpp
--- a/Engine/Plugins/Tests/CQTest/Source/CQTest/Private/Components/MapTestSpawner.cpp
+++ b/Engine/Plugins/Tests/CQTest/Source/CQTest/Private/Components/MapTestSpawner.cpp
@@ -5,12 +5,15 @@
 #if WITH_AUTOMATION_TESTS
 #include "Commands/TestCommands.h"
 #include "Tests/AutomationCommon.h"
+#include "GameDelegates.h"
 #include "GameFramework/PlayerController.h"
 #include "Engine/Engine.h"
 #include "Misc/PackageName.h"
 #include "Misc/Paths.h"
 #include "UObject/Package.h"
 
+DEFINE_LOG_CATEGORY_STATIC(LogMapTest, Log, All);
+
 #if WITH_EDITOR
 #include "Editor.h"
 #include "Editor/UnrealEdEngine.h"
@@ -21,8 +24,6 @@
 #include "Tests/AutomationEditorCommon.h"
 #include "UnrealEdGlobals.h"
 
-DEFINE_LOG_CATEGORY_STATIC(LogMapTest, Log, All);
-
 namespace {
 
 static const FString TempMapDirectory = FPaths::Combine(FPaths::ProjectContentDir(), TEXT("CQTestMapTemp"));
@@ -52,6 +53,17 @@ void CleanupTempResources()
 } //anonymous
 #endif // WITH_EDITOR
 
+FMapTestSpawner::~FMapTestSpawner()
+{
+	// Only explicitly removing the 'EndPlayMapHandle' handle as either the `OnEndPlayMap` gets triggered and the Game/PIE Worlds are no longer valid
+	// Or we are ending the test and the `FSpawnHelper` will handle cleaning up of the GameWorld
+	if (EndPlayMapHandle.IsValid())
+	{
+		FGameDelegates::Get().GetEndPlayMapDelegate().Remove(EndPlayMapHandle);
+		EndPlayMapHandle.Reset();
+	}
+}
+
 TUniquePtr<FMapTestSpawner> FMapTestSpawner::CreateFromTempLevel(FTestCommandBuilder& InCommandBuilder)
 {
 #if WITH_EDITOR
@@ -75,7 +87,7 @@ TUniquePtr<FMapTestSpawner> FMapTestSpawner::CreateFromTempLevel(FTestCommandBui
 	FLevelEditorModule& LevelEditor = FModuleManager::LoadModuleChecked<FLevelEditorModule>("LevelEditor");
 	Spawner->MapChangedHandle = LevelEditor.OnMapChanged().AddRaw(Spawner.Get(), &FMapTestSpawner::OnMapChanged);
 
-	InCommandBuilder.OnTearDown([&]() {
+	InCommandBuilder.OnTearDown([]() {
 		// Create a new map to free up the reference to the map used during testing before cleaning up all temporary resources
 		FAutomationEditorCommonUtils::CreateNewMap();
 		CleanupTempResources();
@@ -104,8 +116,8 @@ void FMapTestSpawner::AddWaitUntilLoadedCommand(FAutomationTestBase* TestRunner)
 	bool bOpened = AutomationOpenMap(LongPackageName, true);
 	check(bOpened);
 
-	ADD_LATENT_AUTOMATION_COMMAND(FWaitUntil(*TestRunner, [&]() -> bool {
-		for (const auto& Context : GEngine->GetWorldContexts())
+	ADD_LATENT_AUTOMATION_COMMAND(FWaitUntil(*TestRunner, [this]() -> bool {
+		for (const FWorldContext& Context : GEngine->GetWorldContexts())
 		{
 			UWorld* World = Context.World();
 			if (!IsValid(World))
@@ -119,6 +131,8 @@ void FMapTestSpawner::AddWaitUntilLoadedCommand(FAutomationTestBase* TestRunner)
 			if (((Context.WorldType == EWorldType::PIE) || (Context.WorldType == EWorldType::Game)) && (WorldMapName.Equals(MapName)))
 			{
 				PieWorld = Context.World();
+				EndPlayMapHandle = FGameDelegates::Get().GetEndPlayMapDelegate().AddRaw(this, &FMapTestSpawner::OnEndPlayMap);
+
 				return true;
 			}
 		}
@@ -146,6 +160,19 @@ APawn* FMapTestSpawner::FindFirstPlayerPawn()
 	return PlayerController->GetPawn();
 }
 
+void FMapTestSpawner::OnEndPlayMap()
+{
+	if (EndPlayMapHandle.IsValid())
+	{
+		UE_LOG(LogMapTest, Verbose, TEXT("Play session has ended."));
+		GameWorld = nullptr;
+		PieWorld = nullptr;
+
+		FGameDelegates::Get().GetEndPlayMapDelegate().Remove(EndPlayMapHandle);
+		EndPlayMapHandle.Reset();
+	}
+}
+
 #if WITH_EDITOR
 void FMapTestSpawner::OnMapChanged(UWorld* World, EMapChangeType ChangeType)
 {
diff --git a/Engine/Plugins/Tests/CQTest/Source/CQTest/Public/Components/MapTestSpawner.h b/Engine/Plugins/Tests/CQTest/Source/CQTest/Public/Components/MapTestSpawner.h
--- a/Engine/Plugins/Tests/CQTest/Source/CQTest/Public/Components/MapTestSpawner.h
+++ b/Engine/Plugins/Tests/CQTest/Source/CQTest/Public/Components/MapTestSpawner.h
@@ -57,6 +57,9 @@ struct CQTEST_API FMapTestSpawner : public FSpawnHelper
 	FMapTestSpawner(const FString& MapDirectory, const FString& MapName)
 		: MapDirectory(MapDirectory), MapName(MapName) {}
 
+	/** Destructor */
+	~FMapTestSpawner();
+
 	/**
 	 * Creates an instance of the MapTestSpawner with a temporary level ready for use.
 	 * 
@@ -70,6 +73,8 @@ struct CQTEST_API FMapTestSpawner : public FSpawnHelper
 	 * Loads the map specified from the MapDirectory and MapName to be prepared for the test.
 	 *
 	 * @param TestRunner - TestRunner used to send the latent command needed for map preparations.
+	 * 
+	 * @note Must be called outside of a latent action. Preferably within BEFORE_TEST.
 	 */
 	void AddWaitUntilLoadedCommand(FAutomationTestBase* TestRunner);
 
@@ -83,6 +88,11 @@ protected:
     virtual UWorld* CreateWorld() override;
 
 private:
+	/**
+	 * Handler called when the PIE session ends.
+	 */
+	void OnEndPlayMap();
+
 #if WITH_EDITOR
 	/**
 	 * Handler called on map changed.
@@ -95,6 +105,7 @@ private:
 	FString MapDirectory;
 	FString MapName;
 	UWorld* PieWorld{ nullptr };
+	FDelegateHandle EndPlayMapHandle;
 };
 
 #endif // WITH_AUTOMATION_TESTS
\ No newline at end of file
diff --git a/Engine/Plugins/Tests/CQTest/Source/CQTestTests/Private/Components/MapSpawnHelperTests.cpp b/Engine/Plugins/Tests/CQTest/Source/CQTestTests/Private/Components/MapSpawnHelperTests.cpp
--- a/Engine/Plugins/Tests/CQTest/Source/CQTestTests/Private/Components/MapSpawnHelperTests.cpp
+++ b/Engine/Plugins/Tests/CQTest/Source/CQTestTests/Private/Components/MapSpawnHelperTests.cpp
@@ -6,6 +6,8 @@
 
 #if WITH_EDITOR && WITH_AUTOMATION_TESTS
 
+#include "Editor.h"
+
 TEST_CLASS_WITH_FLAGS(MapSpawnHelperTests, "TestFramework.CQTest.Map", EAutomationTestFlags::EditorContext | EAutomationTestFlags::EngineFilter)
 {
 	TUniquePtr<FMapTestSpawner> Spawner;
@@ -15,7 +17,7 @@ TEST_CLASS_WITH_FLAGS(MapSpawnHelperTests, "TestFramework.CQTest.Map", EAutomati
 		ASSERT_THAT(IsNotNull(Spawner));
 
 		Spawner->AddWaitUntilLoadedCommand(TestRunner);
-		TestCommandBuilder.Do([&]() {
+		TestCommandBuilder.Do([this]() {
 			// Because we're creating a level for this test, we will also want to populate the level with a Pawn object that can be then tested against
 			Spawner->SpawnActor<APawn>();
 		});
@@ -23,11 +25,19 @@ TEST_CLASS_WITH_FLAGS(MapSpawnHelperTests, "TestFramework.CQTest.Map", EAutomati
 
 	TEST_METHOD(MapSpawner_FindsPlayerSpawn)
 	{
-		TestCommandBuilder.Do([&]() {
+		TestCommandBuilder.Do([this]() {
 			APawn* Player = Spawner->FindFirstPlayerPawn();
 			ASSERT_THAT(IsNotNull(Player));
 		});
 	}
+
+	TEST_METHOD(MapSpawner_PIEEndsEarly)
+	{
+		TestCommandBuilder.Do([this]() { GEditor->RequestEndPlayMap(); })
+		.Until([this]() { return !GEditor->IsPlaySessionInProgress(); })
+		.Until([this]() { return !IsValid(GEditor->PlayWorld); })
+		.Then([this]() { ASSERT_THAT(IsNull(GEditor->PlayWorld)); });
+	}
 };
 
 #endif // WITH_EDITOR && WITH_AUTOMATION_TESTS
\ No newline at end of file
