Adding static_assert to better prevent tests from being hidden from execution.

#rnx

diff --git a/Engine/Plugins/Tests/CQTest/Source/CQTest/Public/CQTest.h b/Engine/Plugins/Tests/CQTest/Source/CQTest/Public/CQTest.h
--- a/Engine/Plugins/Tests/CQTest/Source/CQTest/Public/CQTest.h
+++ b/Engine/Plugins/Tests/CQTest/Source/CQTest/Public/CQTest.h
@@ -172,6 +172,8 @@ struct TTestRunner : public FAutomationTestBase
 	FDelegateHandle BeforeAllDelegate{};
 	FDelegateHandle AfterAllDelegate{};
 
+	TTestRunner() = delete;
+
 protected:
 	FString GetBeautifiedTestName() const override;
 	uint32 GetRequiredDeviceNum() const override;
@@ -228,7 +230,17 @@ struct TTest : TBaseTest<AsserterType>
 	struct F##_ClassName##_Runner : public TTestRunner<_AsserterType>                                                                  \
 	{                                                                                                                                  \
 		F##_ClassName##_Runner()                                                                                                       \
-			: TTestRunner(#_ClassName, __LINE__, __FILE__, _TestDir, _TestFlags, TTest<_ClassName, _AsserterType>::CreateTestClass) {} \
+			: TTestRunner(#_ClassName, __LINE__, __FILE__, _TestDir, _TestFlags, TTest<_ClassName, _AsserterType>::CreateTestClass) {  \
+				static_assert(  !!((_TestFlags) & EAutomationTestFlags_ApplicationContextMask),                                        \
+								"CQTest has no application flag and will not run. See AutomationTest.h.");                             \
+				static_assert(  !!(((_TestFlags) & EAutomationTestFlags_FilterMask) == EAutomationTestFlags::SmokeFilter) ||           \
+								!!(((_TestFlags) & EAutomationTestFlags_FilterMask) == EAutomationTestFlags::EngineFilter) ||          \
+								!!(((_TestFlags) & EAutomationTestFlags_FilterMask) == EAutomationTestFlags::ProductFilter) ||         \
+								!!(((_TestFlags) & EAutomationTestFlags_FilterMask) == EAutomationTestFlags::PerfFilter) ||            \
+								!!(((_TestFlags) & EAutomationTestFlags_FilterMask) == EAutomationTestFlags::StressFilter) ||          \
+								!!(((_TestFlags) & EAutomationTestFlags_FilterMask) == EAutomationTestFlags::NegativeFilter),          \
+								"All CQTests must have exactly 1 filter type specified. See AutomationTest.h.");                       \
+		}                                                                                                                              \
 	};                                                                                                                                 \
 	F##_ClassName##_Runner _ClassName##_RunnerInstance;                                                                                \
 	struct _ClassName : public _BaseClass<_ClassName, _AsserterType>
