Fixing cause of memory corruption caused by delegates not being cleaned up correctly.

#rnx

diff --git a/Engine/Plugins/Tests/CQTest/Source/CQTest/Private/Components/MapTestSpawner.cpp b/Engine/Plugins/Tests/CQTest/Source/CQTest/Private/Components/MapTestSpawner.cpp
--- a/Engine/Plugins/Tests/CQTest/Source/CQTest/Private/Components/MapTestSpawner.cpp
+++ b/Engine/Plugins/Tests/CQTest/Source/CQTest/Private/Components/MapTestSpawner.cpp
@@ -44,9 +44,9 @@ FString GenerateUniqueMapName()
  */
 void CleanupTempResources()
 {
-	bool bDirectoryMustExist = true;
-	bool bRemoveRecursively = true;
-	bool bWasDeleted = IFileManager::Get().DeleteDirectory(*TempMapDirectory, bDirectoryMustExist, bRemoveRecursively);
+	const bool bDirectoryMustExist = true;
+	const bool bRemoveRecursively = true;
+	const bool bWasDeleted = IFileManager::Get().DeleteDirectory(*TempMapDirectory, bDirectoryMustExist, bRemoveRecursively);
 	check(bWasDeleted);
 }
 
@@ -82,11 +82,6 @@ TUniquePtr<FMapTestSpawner> FMapTestSpawner::CreateFromTempLevel(FTestCommandBui
 	check(bWasTempLevelCreated);
 
 	TUniquePtr<FMapTestSpawner> Spawner = MakeUnique<FMapTestSpawner>(TempMapDirectory, MapName);
-	
-	// Register Map Change Events
-	FLevelEditorModule& LevelEditor = FModuleManager::LoadModuleChecked<FLevelEditorModule>("LevelEditor");
-	Spawner->MapChangedHandle = LevelEditor.OnMapChanged().AddRaw(Spawner.Get(), &FMapTestSpawner::OnMapChanged);
-
 	InCommandBuilder.OnTearDown([]() {
 		// Create a new map to free up the reference to the map used during testing before cleaning up all temporary resources
 		FAutomationEditorCommonUtils::CreateNewMap();
@@ -162,7 +157,7 @@ APawn* FMapTestSpawner::FindFirstPlayerPawn()
 
 void FMapTestSpawner::OnEndPlayMap()
 {
-	if (EndPlayMapHandle.IsValid())
+	if (!IsValid(GEngine->GetCurrentPlayWorld()))
 	{
 		UE_LOG(LogMapTest, Verbose, TEXT("Play session has ended."));
 		GameWorld = nullptr;
@@ -173,19 +168,4 @@ void FMapTestSpawner::OnEndPlayMap()
 	}
 }
 
-#if WITH_EDITOR
-void FMapTestSpawner::OnMapChanged(UWorld* World, EMapChangeType ChangeType)
-{
-	if (PieWorld && ChangeType == EMapChangeType::TearDownWorld)
-	{
-		UE_LOG(LogMapTest, Verbose, TEXT("Map used by the Spawner has been changed."));
-		GameWorld = nullptr;
-		PieWorld = nullptr;
-
-		FLevelEditorModule& LevelEditor = FModuleManager::LoadModuleChecked<FLevelEditorModule>("LevelEditor");
-		LevelEditor.OnMapChanged().Remove(MapChangedHandle);
-	}
-}
-#endif // WITH_EDITOR
-
 #endif // WITH_AUTOMATION_TESTS
\ No newline at end of file
diff --git a/Engine/Plugins/Tests/CQTest/Source/CQTest/Public/Components/MapTestSpawner.h b/Engine/Plugins/Tests/CQTest/Source/CQTest/Public/Components/MapTestSpawner.h
--- a/Engine/Plugins/Tests/CQTest/Source/CQTest/Public/Components/MapTestSpawner.h
+++ b/Engine/Plugins/Tests/CQTest/Source/CQTest/Public/Components/MapTestSpawner.h
@@ -88,20 +88,9 @@ protected:
     virtual UWorld* CreateWorld() override;
 
 private:
-	/**
-	 * Handler called when the PIE session ends.
-	 */
+	/** Handler called when the PIE session ends. */
 	void OnEndPlayMap();
 
-#if WITH_EDITOR
-	/**
-	 * Handler called on map changed.
-	 */
-	void OnMapChanged(UWorld* World, EMapChangeType ChangeType);
-
-	FDelegateHandle MapChangedHandle;
-#endif // WITH_EDITOR
-
 	FString MapDirectory;
 	FString MapName;
 	UWorld* PieWorld{ nullptr };
