Add support for async operations in CQTest within TestCommandBuilder

#rnx

diff --git a/Engine/Plugins/Tests/CQTest/Source/CQTestTests/Private/Commands/CommandBuilderTests.cpp b/Engine/Plugins/Tests/CQTest/Source/CQTestTests/Private/Commands/CommandBuilderTests.cpp
--- a/Engine/Plugins/Tests/CQTest/Source/CQTestTests/Private/Commands/CommandBuilderTests.cpp
+++ b/Engine/Plugins/Tests/CQTest/Source/CQTestTests/Private/Commands/CommandBuilderTests.cpp
@@ -77,4 +77,656 @@ TEST_CLASS(CommandBuilderTests, "TestFramework.CQTest.Core")
 		ASSERT_THAT(IsNotNull(command));
 		ASSERT_THAT(IsNull(secondTime));
 	}
+	
+	TEST_METHOD(DoAsync_ThenBuild_IncludesCommand)
+	{
+		TPromise<bool> Promise;
+		bool bAsyncActionInvoked = false;
+
+		auto Command = CommandBuilder.DoAsync<bool>(
+			[&]()
+			{
+				bAsyncActionInvoked = true;
+				return TAsyncResult<bool>(Promise.GetFuture(), nullptr, nullptr);
+			}
+		).Build();
+		
+		// Start async action
+		ASSERT_THAT(IsFalse(bAsyncActionInvoked));
+		ASSERT_THAT(IsFalse(Command->Update()));
+		ASSERT_THAT(IsTrue(bAsyncActionInvoked));
+		
+		// Wait for async result
+		ASSERT_THAT(IsFalse(Command->Update()));
+		Promise.SetValue(true);
+		ASSERT_THAT(IsTrue(Command->Update()));
+	}
+	
+	TEST_METHOD(DoAsync_WithResultCallback_ThenBuild_IncludesCommand)
+	{
+		TPromise<bool> Promise;
+		bool bAsyncActionInvoked = false;
+		bool bResultCallbackInvoked = false;
+
+		auto Command = CommandBuilder.DoAsync<bool>(
+			[&]()
+			{
+				bAsyncActionInvoked = true;
+				return TAsyncResult<bool>(Promise.GetFuture(), nullptr, nullptr); 
+			},
+			[&](bool)
+			{
+				bResultCallbackInvoked = true; 
+			}
+		).Build();
+		
+		// Start async action
+		ASSERT_THAT(IsFalse(bAsyncActionInvoked));
+		ASSERT_THAT(IsFalse(Command->Update()));
+		ASSERT_THAT(IsTrue(bAsyncActionInvoked));
+		
+		// Wait for async result
+		ASSERT_THAT(IsFalse(Command->Update()));
+		Promise.SetValue(true);
+		ASSERT_THAT(IsFalse(Command->Update()));
+		
+		// Handle result
+		ASSERT_THAT(IsFalse(bResultCallbackInvoked));
+		ASSERT_THAT(IsTrue(Command->Update()));
+		ASSERT_THAT(IsTrue(bResultCallbackInvoked));
+	}
+
+	TEST_METHOD(DoAsync_WaitsAsyncResultForSpecifiedDuration)
+	{
+		TPromise<bool> Promise;
+		const FTimespan AsyncActionTimeout = FTimespan::FromMilliseconds(200);
+		const FDateTime StartTime = FDateTime::UtcNow();
+		const FDateTime MaxEndTime = StartTime + FTimespan::FromMilliseconds(500);
+
+		auto Command = CommandBuilder.DoAsync<bool>(
+			[&]() { return TAsyncResult<bool>(Promise.GetFuture(), nullptr, nullptr); },
+			AsyncActionTimeout
+		).Build();
+
+		bool bDone = false;
+		while (!bDone && FDateTime::UtcNow() < MaxEndTime)
+		{
+			bDone = Command->Update();
+		}
+
+		const FDateTime EndTime = FDateTime::UtcNow();
+		ClearExpectedError(*TestRunner, "Latent command timed out");
+		// Set value to destroy the promise correctly
+		Promise.SetValue(true);
+
+		ASSERT_THAT(IsTrue(bDone));
+		ASSERT_THAT(IsTrue(EndTime >= StartTime + AsyncActionTimeout));
+	}
+
+	TEST_METHOD(DoAsync_ShouldProcessResultOfType_Void)
+	{
+		TPromise<void> Promise;
+
+		auto Command = CommandBuilder.DoAsync<void>(
+			[&]()
+			{
+				return TAsyncResult<void>(Promise.GetFuture(), nullptr, nullptr);
+			}
+		).Build();
+		
+		// Start async action
+		ASSERT_THAT(IsFalse(Command->Update()));
+		
+		// Wait for async result
+		Promise.SetValue();
+		ASSERT_THAT(IsTrue(Command->Update()));
+	}
+
+	TEST_METHOD(DoAsync_ShouldProcessResultOfType_Class)
+	{
+		class FTestContainer
+		{
+		public:
+			FTestContainer(int32 InValue) : Value(InValue) {}
+			int32 Value;
+		};
+
+		TPromise<FTestContainer> Promise;
+		FTestContainer Result(0);
+
+		auto Command = CommandBuilder.DoAsync<FTestContainer>(
+			[&]()
+			{
+				return TAsyncResult<FTestContainer>(Promise.GetFuture(), nullptr, nullptr);
+			},
+			[&](const FTestContainer& InResult)
+			{
+				Result = InResult;
+			}
+		).Build();
+		
+		// Start async action
+		ASSERT_THAT(IsFalse(Command->Update()));
+		
+		const FTestContainer ExpectedResult(567);
+		Promise.SetValue(ExpectedResult);
+		// Need two updates: the first to verify that the async result is ready, the second to invoke the callback.
+		bool bCompleted = Command->Update() || Command->Update();
+		ASSERT_THAT(IsTrue(bCompleted));
+		ASSERT_THAT(AreEqual(ExpectedResult.Value, Result.Value));
+	}
+
+	TEST_METHOD(DoAsync_ShouldProcessResultOfType_Reference)
+	{
+		TPromise<int32&> Promise;
+		int32 Value = 12345;
+		const int32 ExpectedValue = 6789;
+
+		auto Command = CommandBuilder.DoAsync<int32&>(
+			[&]()
+			{
+				return TAsyncResult<int32&>(Promise.GetFuture(), nullptr, nullptr);
+			},
+			[&](int32& InResult)
+			{
+				InResult = ExpectedValue;
+			}
+		).Build();
+		
+		// Start async action
+		ASSERT_THAT(IsFalse(Command->Update()));
+
+		Promise.SetValue(Value);
+		// Need two updates: the first to verify that the async result is ready, the second to invoke the callback.
+		bool bCompleted = Command->Update() || Command->Update();
+		ASSERT_THAT(IsTrue(bCompleted));
+		ASSERT_THAT(AreEqual(ExpectedValue, Value));
+	}
+
+	TEST_METHOD(UntilAsync_ThenBuild_IncludesCommand)
+	{
+		TPromise<bool> Promise;
+		bool bAsyncActionInvoked = false;
+		bool bConditionChecked = false;
+		bool bConditionResult = false;
+
+		auto Command = CommandBuilder.UntilAsync<bool>(
+			[&]()
+			{
+				bAsyncActionInvoked = true;
+				return TAsyncResult<bool>(Promise.GetFuture(), nullptr, nullptr);
+			},
+			[&](bool)
+			{
+				bConditionChecked = true;
+				return bConditionResult;
+			}
+		).Build();
+		
+		// Start async action
+		ASSERT_THAT(IsFalse(bAsyncActionInvoked));
+		ASSERT_THAT(IsFalse(Command->Update()));
+		ASSERT_THAT(IsTrue(bAsyncActionInvoked));
+		
+		// Wait for async result
+		ASSERT_THAT(IsFalse(Command->Update()));
+		Promise.SetValue(true);
+		ASSERT_THAT(IsFalse(Command->Update()));
+		
+		// Start checking condition
+		ASSERT_THAT(IsFalse(bConditionChecked));
+		ASSERT_THAT(IsFalse(Command->Update()));
+		ASSERT_THAT(IsTrue(bConditionChecked));
+		
+		// Stop execution when the condition is met
+		bConditionResult = true;
+		bConditionChecked = false;
+		ASSERT_THAT(IsTrue(Command->Update()));
+		ASSERT_THAT(IsTrue(bConditionChecked));
+	}
+
+	TEST_METHOD(UntilAsync_WaitsAsyncResultForSpecifiedDuration)
+	{
+		TPromise<bool> Promise;
+		const FTimespan AsyncActionTimeout = FTimespan::FromMilliseconds(200);
+		const FTimespan ConditionTimeout = FTimespan::FromSeconds(10);
+		const FDateTime StartTime = FDateTime::UtcNow();
+		const FDateTime MaxEndTime = StartTime + FTimespan::FromMilliseconds(500);
+
+		auto Command = CommandBuilder.UntilAsync<bool>(
+			[&]() { return TAsyncResult<bool>(Promise.GetFuture(), nullptr, nullptr); },
+			[](bool) { return false; },
+			AsyncActionTimeout,
+			ConditionTimeout
+		).Build();
+
+		bool bDone = false;
+		while (!bDone && FDateTime::UtcNow() < MaxEndTime)
+		{
+			bDone = Command->Update();
+		}
+
+		const FDateTime EndTime = FDateTime::UtcNow();
+		ClearExpectedError(*TestRunner, "Latent command timed out");
+		// Set value to destroy the promise correctly
+		Promise.SetValue(true);
+		
+		ASSERT_THAT(IsTrue(bDone));
+		ASSERT_THAT(IsTrue(EndTime >= StartTime + AsyncActionTimeout));
+	}
+
+	TEST_METHOD(UntilAsync_WaitsConditionForSpecifiedDuration)
+	{
+		TPromise<bool> Promise;
+		const FTimespan AsyncActionTimeout = FTimespan::FromSeconds(10);
+		const FTimespan ConditionTimeout = FTimespan::FromMilliseconds(200);
+		const FDateTime StartTime = FDateTime::UtcNow();
+		const FDateTime MaxEndTime = StartTime + FTimespan::FromMilliseconds(500);
+
+		auto Command = CommandBuilder.UntilAsync<bool>(
+			[&]() { return TAsyncResult<bool>(Promise.GetFuture(), nullptr, nullptr); },
+			[](bool) { return false; },
+			AsyncActionTimeout,
+			ConditionTimeout
+		).Build();
+		
+		Promise.SetValue(true);
+
+		bool bDone = false;
+		while (!bDone && FDateTime::UtcNow() < MaxEndTime)
+		{
+			bDone = Command->Update();
+		}
+
+		const FDateTime EndTime = FDateTime::UtcNow();
+		ClearExpectedError(*TestRunner, "Latent command timed out");
+		
+		ASSERT_THAT(IsTrue(bDone));
+		ASSERT_THAT(IsTrue(EndTime >= StartTime + ConditionTimeout));
+	}
+
+	TEST_METHOD(UntilAsync_ShouldProcessResultOfType_Class)
+	{
+		class FTestContainer
+		{
+		public:
+			FTestContainer(int32 InValue) : Value(InValue) {}
+			int32 Value;
+		};
+
+		TPromise<FTestContainer> Promise;
+		FTestContainer Result(0);
+
+		auto Command = CommandBuilder.UntilAsync<FTestContainer>(
+			[&]()
+			{
+				return TAsyncResult<FTestContainer>(Promise.GetFuture(), nullptr, nullptr);
+			},
+			[&](const FTestContainer& InResult)
+			{
+				Result = InResult;
+				return true;
+			}
+		).Build();
+		
+		// Start async action
+		ASSERT_THAT(IsFalse(Command->Update()));
+		
+		const FTestContainer ExpectedResult(567);
+		Promise.SetValue(ExpectedResult);
+		// Need two updates: the first to verify that the async result is ready, the second to invoke the callback.
+		bool bCompleted = Command->Update() || Command->Update();
+		ASSERT_THAT(IsTrue(bCompleted));
+		ASSERT_THAT(AreEqual(ExpectedResult.Value, Result.Value));
+	}
+
+	TEST_METHOD(UntilAsync_ShouldProcessResultOfType_Reference)
+	{
+		TPromise<int32&> Promise;
+		int32 Value = 12345;
+		const int32 ExpectedValue = 6789;
+
+		auto Command = CommandBuilder.UntilAsync<int32&>(
+			[&]()
+			{
+				return TAsyncResult<int32&>(Promise.GetFuture(), nullptr, nullptr);
+			},
+			[&](int32& Result)
+			{
+				Result = ExpectedValue;
+				return true;
+			}
+		).Build();
+		
+		// Start async action
+		ASSERT_THAT(IsFalse(Command->Update()));
+
+		Promise.SetValue(Value);
+		// Need two updates: the first to verify that the async result is ready, the second to invoke the callback.
+		bool bCompleted = Command->Update() || Command->Update();
+		ASSERT_THAT(IsTrue(bCompleted));
+		ASSERT_THAT(AreEqual(ExpectedValue, Value));
+	}
+	
+	TEST_METHOD(ThenAsync_ThenBuild_IncludesCommand)
+	{
+		TPromise<bool> Promise;
+		bool bAsyncActionInvoked = false;
+
+		auto Command = CommandBuilder.ThenAsync<bool>(
+			[&]()
+			{
+				bAsyncActionInvoked = true;
+				return TAsyncResult<bool>(Promise.GetFuture(), nullptr, nullptr);
+			}
+		).Build();
+		
+		// Start async action
+		ASSERT_THAT(IsFalse(bAsyncActionInvoked));
+		ASSERT_THAT(IsFalse(Command->Update()));
+		ASSERT_THAT(IsTrue(bAsyncActionInvoked));
+		
+		// Wait for async result
+		ASSERT_THAT(IsFalse(Command->Update()));
+		Promise.SetValue(true);
+		ASSERT_THAT(IsTrue(Command->Update()));
+	}
+	
+	TEST_METHOD(ThenAsync_WithResultCallback_ThenBuild_IncludesCommand)
+	{
+		TPromise<bool> Promise;
+		bool bAsyncActionInvoked = false;
+		bool bResultCallbackInvoked = false;
+
+		auto Command = CommandBuilder.ThenAsync<bool>(
+			[&]()
+			{
+				bAsyncActionInvoked = true;
+				return TAsyncResult<bool>(Promise.GetFuture(), nullptr, nullptr); 
+			},
+			[&](bool)
+			{
+				bResultCallbackInvoked = true; 
+			}
+		).Build();
+		
+		// Start async action
+		ASSERT_THAT(IsFalse(bAsyncActionInvoked));
+		ASSERT_THAT(IsFalse(Command->Update()));
+		ASSERT_THAT(IsTrue(bAsyncActionInvoked));
+		
+		// Wait for async result
+		ASSERT_THAT(IsFalse(Command->Update()));
+		Promise.SetValue(true);
+		ASSERT_THAT(IsFalse(Command->Update()));
+		
+		// Handle result
+		ASSERT_THAT(IsFalse(bResultCallbackInvoked));
+		ASSERT_THAT(IsTrue(Command->Update()));
+		ASSERT_THAT(IsTrue(bResultCallbackInvoked));
+	}
+
+	TEST_METHOD(ThenAsync_WaitsAsyncResultForSpecifiedDuration)
+	{
+		TPromise<bool> Promise;
+		const FTimespan AsyncActionTimeout = FTimespan::FromMilliseconds(200);
+		const FDateTime StartTime = FDateTime::UtcNow();
+		const FDateTime MaxEndTime = StartTime + FTimespan::FromMilliseconds(500);
+
+		auto Command = CommandBuilder.ThenAsync<bool>(
+			[&]() { return TAsyncResult<bool>(Promise.GetFuture(), nullptr, nullptr); },
+			AsyncActionTimeout
+		).Build();
+
+		bool bDone = false;
+		while (!bDone && FDateTime::UtcNow() < MaxEndTime)
+		{
+			bDone = Command->Update();
+		}
+
+		const FDateTime EndTime = FDateTime::UtcNow();
+		ClearExpectedError(*TestRunner, "Latent command timed out");
+		// Set value to destroy the promise correctly
+		Promise.SetValue(true);
+
+		ASSERT_THAT(IsTrue(bDone));
+		ASSERT_THAT(IsTrue(EndTime >= StartTime + AsyncActionTimeout));
+	}
+};
+
+/* These tests illustrate different approaches to running functions that return TAsyncResult within Command Builder */
+TEST_CLASS(CommandBuilderForAsyncResultTests, "TestFramework.CQTest.Core")
+{
+	class FFakeAsyncTask
+	{
+	public:
+		FFakeAsyncTask(FAutomationTestBase& InTestRunner, int InValue, int32 InDuration) :
+			TestRunner(InTestRunner),
+			Value(InValue),
+			Duration(InDuration)
+		{
+			check(Duration > 0);
+		}
+
+		~FFakeAsyncTask()
+		{
+			if (IsRunning) { Promise->SetValue({}); }
+		}
+
+		TAsyncResult<int> ProduceValue()
+		{
+			if (TestRunner.AddErrorIfFalse(!IsRunning, TEXT("Async task has already been started")))
+			{
+				IsRunning = true;
+				Promise = MakeShared<TPromise<int>>();
+				return TAsyncResult<int>(Promise->GetFuture(), nullptr, nullptr);
+			}
+			return TAsyncResult<int>();
+		}
+
+		void Update()
+		{
+			if (!IsRunning) { return; }
+
+			if (--Duration > 0) { return; }
+
+			Promise->SetValue(Value);
+			IsRunning = false;
+		}
+
+	private:
+		FAutomationTestBase& TestRunner;
+		TSharedPtr<TPromise<int>> Promise;
+		int Value;
+		int32 Duration;
+		bool IsRunning = false;
+	};
+
+	class FFakeBackgroundTask
+	{
+	public:
+		FFakeBackgroundTask(FAutomationTestBase& InTestRunner, int InValue, int32 InDuration) :
+			TestRunner(InTestRunner),
+			Value(InValue),
+			Duration(InDuration) {}
+
+		bool IsReady(int InValue)
+		{
+			if (TestRunner.AddErrorIfFalse(Value == InValue,
+				FString::Printf(TEXT("Incorrect value. Expected: %d, actual: %d"), Value, InValue)))
+			{
+				return Duration == 0;
+			}
+			return false;
+		}
+
+		void Update() { if (Duration > 0) { Duration--; } }
+
+	private:
+		FAutomationTestBase& TestRunner;
+		const int Value;
+		int32 Duration;
+	};
+
+	FTestCommandBuilder CommandBuilder{ *TestRunner };
+
+	/* Execute an async task without checking the return value, using a sequence of general-purpose commands. */
+	TEST_METHOD(Execute_StepByStep)
+	{
+		const int32 TaskDurationInTicks = 5;
+		FFakeAsyncTask AsyncTask(*TestRunner, 0, TaskDurationInTicks);
+
+		TAsyncResult<int> AsyncResult;
+
+		auto Command = CommandBuilder
+			.Do(TEXT("Start producing value"), [&]() { AsyncResult = AsyncTask.ProduceValue(); })
+			.Until(TEXT("Value produced"), [&]() { return AsyncResult.GetFuture().IsReady(); })
+			.Build();
+
+		bool bDone = false;
+		int32 Timeout = TaskDurationInTicks + 1;
+		while (!bDone && Timeout--)
+		{
+			bDone = Command->Update();
+			AsyncTask.Update();
+		}
+
+		ASSERT_THAT(IsTrue(bDone));
+	}
+	
+	/* Execute an async task without checking the return value, using a DoAsync command. */
+	TEST_METHOD(Execute_ByDoAsync)
+	{
+		const int32 TaskDurationInTicks = 5;
+		FFakeAsyncTask Task(*TestRunner, 0, TaskDurationInTicks);
+
+		auto Command = CommandBuilder
+			.DoAsync<int>(TEXT("Produce value"), [&]() { return Task.ProduceValue(); })
+			.Build();
+
+		bool bDone = false;
+		int32 Timeout = TaskDurationInTicks + 1;
+		while (!bDone && Timeout--)
+		{
+			bDone = Command->Update();
+			Task.Update();
+		}
+
+		ASSERT_THAT(IsTrue(bDone));
+	}
+	
+	/* Execute an async task and retrieve the return value using a sequence of general-purpose commands. */
+	TEST_METHOD(ExecuteAndGetResult_StepByStep)
+	{
+		const int ExpectedValue = 123;
+		const int32 TaskDurationInTicks = 5;
+		FFakeAsyncTask AsyncTask(*TestRunner, ExpectedValue, TaskDurationInTicks);
+
+		TAsyncResult<int> AsyncResult;
+		int Result = 0;
+
+		auto Command = CommandBuilder
+			.Do(TEXT("Start producing value"), [&]() { AsyncResult = AsyncTask.ProduceValue(); })
+			.Until(TEXT("Value produced"), [&]() { return AsyncResult.GetFuture().IsReady(); })
+			.Then(TEXT("Save value"), [&]() { Result = AsyncResult.GetFuture().Get(); })
+			.Build();
+
+		bool bDone = false;
+		int32 Timeout = TaskDurationInTicks + 2;
+		while (!bDone && Timeout--)
+		{
+			bDone = Command->Update();
+			AsyncTask.Update();
+		}
+
+		ASSERT_THAT(AreEqual(ExpectedValue, Result));
+	}
+	
+	/* Execute an async task and retrieve the return value using a DoAsync command. */
+	TEST_METHOD(ExecuteAndGetResult_ByDoAsync)
+	{
+		const int ExpectedValue = 456;
+		const int32 TaskDurationInTicks = 5;
+		FFakeAsyncTask Task(*TestRunner, ExpectedValue, TaskDurationInTicks);
+
+		int Result = 0;
+
+		auto Command = CommandBuilder
+			.DoAsync<int>(TEXT("Produce value"),
+				[&]() { return Task.ProduceValue(); },
+				[&](int InResult) { Result = InResult; }
+			)
+			.Build();
+
+		bool bDone = false;
+		int32 Timeout = TaskDurationInTicks + 2;
+		while (!bDone && Timeout--)
+		{
+			bDone = Command->Update();
+			Task.Update();
+		}
+
+		ASSERT_THAT(AreEqual(ExpectedValue, Result));
+	}
+	
+	/* Execute an async task and wait for the condition specified by the return value, using a sequence of general-purpose commands. */
+	TEST_METHOD(ExecuteAndWait_StepByStep)
+	{
+		const int ProducedValue = 789;
+		const int32 AsyncTaskDuration = 5;
+		const int32 BackgroundTaskDuration = 10;
+
+		FFakeAsyncTask AsyncTask(*TestRunner, ProducedValue, AsyncTaskDuration);
+		FFakeBackgroundTask BackgroundTask(*TestRunner, ProducedValue, BackgroundTaskDuration);
+
+		TAsyncResult<int> AsyncResult;
+
+		auto Command = CommandBuilder
+			.Do(TEXT("Start producing value"), [&]() { AsyncResult = AsyncTask.ProduceValue(); })
+			.Until(TEXT("Value produced"), [&]() { return AsyncResult.GetFuture().IsReady(); })
+			.Until(TEXT("Resource is ready"), [&]() { return BackgroundTask.IsReady(AsyncResult.GetFuture().Get()); })
+			.Build();
+
+		bool bDone = false;
+		int32 Timeout = BackgroundTaskDuration + 1;
+		while (!bDone && Timeout--)
+		{
+			bDone = Command->Update();
+			AsyncTask.Update();
+			BackgroundTask.Update();
+		}
+
+		ASSERT_THAT(IsTrue(bDone));
+		ASSERT_THAT(IsTrue(BackgroundTask.IsReady(ProducedValue)));
+	}
+	
+	/* Execute an async task and wait for the condition specified by the return value, using an UntilAsync command. */
+	TEST_METHOD(ExecuteAndWait_ByUntilAsync)
+	{
+		const int ProducedValue = 987;
+		const int32 AsyncTaskDuration = 5;
+		const int32 BackgroundTaskDuration = 10;
+
+		FFakeAsyncTask AsyncTask(*TestRunner, ProducedValue, AsyncTaskDuration);
+		FFakeBackgroundTask BackgroundTask(*TestRunner, ProducedValue, BackgroundTaskDuration);
+
+		auto Command = CommandBuilder
+			.UntilAsync<int>(TEXT("Produced resource is ready"),
+				[&]() { return AsyncTask.ProduceValue(); },
+				[&](int InResult) { return BackgroundTask.IsReady(InResult); }
+			)
+			.Build();
+
+		bool bDone = false;
+		int32 Timeout = BackgroundTaskDuration + 1;
+		while (!bDone && Timeout--)
+		{
+			bDone = Command->Update();
+			AsyncTask.Update();
+			BackgroundTask.Update();
+		}
+
+		ASSERT_THAT(IsTrue(bDone));
+		ASSERT_THAT(IsTrue(BackgroundTask.IsReady(ProducedValue)));
+	}
 };
\ No newline at end of file
diff --git a/Engine/Plugins/Tests/CQTest/Source/CQTestTests/Private/Commands/TestCommandsTests.cpp b/Engine/Plugins/Tests/CQTest/Source/CQTestTests/Private/Commands/TestCommandsTests.cpp
--- a/Engine/Plugins/Tests/CQTest/Source/CQTestTests/Private/Commands/TestCommandsTests.cpp
+++ b/Engine/Plugins/Tests/CQTest/Source/CQTestTests/Private/Commands/TestCommandsTests.cpp
@@ -1,6 +1,8 @@
 // Copyright Epic Games, Inc. All Rights Reserved.
 
 #include "CQTest.h"
+#include "CQTestUnitTestHelper.h"
+
 TEST_CLASS(RunSequenceBasicTests, "TestFramework.CQTest.Core")
 {
 	TSharedPtr<IAutomationLatentCommand> Cmd1;
@@ -299,4 +301,343 @@ TEST_CLASS(RunSequenceTests, "TestFramework.CQTest.Core")
 			return test->CommandLog.Num() == 9;
 		};
 	}
+};
+
+class FFakeAsyncTask
+{
+public:
+	FFakeAsyncTask(FAutomationTestBase& InTestRunner) : TestRunner{ InTestRunner } {}
+
+	~FFakeAsyncTask()
+	{
+		if (InProgress()) { Complete(0); }
+	}
+
+	TAsyncResult<int> Start()
+	{
+		if (TestRunner.AddErrorIfFalse(!IsRunning, TEXT("Async task has already been started")))
+		{
+			IsRunning = true;
+			Promise = MakeShared<TPromise<int>>();
+			return TAsyncResult<int>(Promise->GetFuture(), nullptr, nullptr);
+		}
+		return TAsyncResult<int>();
+	}
+
+	void Complete(int Value)
+	{
+		if (TestRunner.AddErrorIfFalse(IsRunning, TEXT("Cannot set the async task result if it's not running")))
+		{
+			IsRunning = false;
+			Promise->SetValue(Value);
+		}
+	}
+
+	bool InProgress() const { return IsRunning; }
+
+private:
+	bool IsRunning = false;
+	FAutomationTestBase& TestRunner;
+	TSharedPtr<TPromise<int>> Promise;
+};
+
+TEST_CLASS(AsyncExecuteBasicTests, "TestFramework.CQTest.Core")
+{
+	TEST_METHOD(AsyncExecute_InvokesAsyncActionWhenUpdated)
+	{
+		FFakeAsyncTask Task(*TestRunner);
+
+		TAsyncExecute<int> AsyncExecute = TAsyncExecute<int>(
+			*TestRunner,
+			[&]() { return Task.Start(); }
+		);
+
+		ASSERT_THAT(IsFalse(Task.InProgress(), TEXT("Default state of async task is invalid")));
+		ASSERT_THAT(IsFalse(AsyncExecute.Update(), TEXT("Command stopped execution early")));
+		ASSERT_THAT(IsTrue(Task.InProgress(), TEXT("Async task hasn't been started")));
+	}
+
+	TEST_METHOD(AsyncExecute_InvokesAsyncActionOnce)
+	{
+		FFakeAsyncTask Task(*TestRunner);
+		int32 Counter = 0;
+
+		TAsyncExecute<int> AsyncExecute = TAsyncExecute<int>(
+			*TestRunner,
+			[&]()
+			{
+				Counter++;
+				return Task.Start();
+			}
+		);
+		
+		int32 Ticks = 3;
+		bool bDone = false;
+		while (!bDone && Ticks--)
+		{
+			bDone = AsyncExecute.Update();
+		}
+		ASSERT_THAT(IsFalse(bDone, TEXT("Command stopped execution early")));
+		ASSERT_THAT(AreEqual(1, Counter, TEXT("Async action was invoked multiple times")));
+	}
+
+	TEST_METHOD(AsyncExecute_WaitsUntilAsyncResultIsReady)
+	{
+		FFakeAsyncTask Task(*TestRunner);
+
+		TAsyncExecute<int> AsyncExecute = TAsyncExecute<int>(
+			*TestRunner,
+			[&]() { return Task.Start(); }
+		);
+
+		// Start async task
+		ASSERT_THAT(IsFalse(AsyncExecute.Update(), TEXT("Command stopped execution early")));
+
+		bool bDone = false;
+		for (int32 i = 0; !bDone && i < 10; i++)
+		{
+			bDone = AsyncExecute.Update();
+		}
+		ASSERT_THAT(IsFalse(bDone, TEXT("Command stopped execution before the async task was completed")));
+
+		Task.Complete(0);
+		ASSERT_THAT(IsTrue(AsyncExecute.Update(), TEXT("Command failed to stop execution after the async task was completed")));
+	}
+
+	TEST_METHOD(AsyncExecute_FExecute_CompletesWhenAsyncActionIsCompleted)
+	{
+		FFakeAsyncTask Task(*TestRunner);
+
+		TAsyncExecute<int, FExecute> AsyncExecute = TAsyncExecute<int, FExecute>(
+			*TestRunner,
+			[&]() { return Task.Start(); },
+			FTimespan::FromSeconds(1),
+			nullptr,
+			[](int) {}
+		);
+
+		// Start async task
+		ASSERT_THAT(IsFalse(AsyncExecute.Update(), TEXT("Command stopped execution early")));
+
+		Task.Complete(0);
+		// Need two updates: the first to verify that the async result is ready, the second to invoke the callback.
+		bool bCompleted = AsyncExecute.Update() || AsyncExecute.Update();
+		ASSERT_THAT(IsTrue(bCompleted, TEXT("Command failed to stop execution after the async task was completed")));
+	}
+
+	TEST_METHOD(AsyncExecute_FExecute_InvokesResultCallbackWhenAsyncActionIsCompleted)
+	{
+		FFakeAsyncTask Task(*TestRunner);
+		bool bInvoked = false;
+
+		TAsyncExecute<int, FExecute> AsyncExecute = TAsyncExecute<int, FExecute>(
+			*TestRunner,
+			[&]() { return Task.Start(); },
+			FTimespan::FromSeconds(1),
+			nullptr,
+			[&](int) { bInvoked = true; }
+		);
+
+		// Start async task
+		ASSERT_THAT(IsFalse(AsyncExecute.Update(), TEXT("Command stopped execution early")));
+
+		ASSERT_THAT(IsFalse(bInvoked, TEXT("Callback was invoked early")));
+		Task.Complete(0);
+
+		// Need two updates: the first to verify that the async result is ready, the second to invoke the callback.
+		bool bCompleted = AsyncExecute.Update() || AsyncExecute.Update();
+		ASSERT_THAT(IsTrue(bCompleted, TEXT("Command did not stop execution")));
+		ASSERT_THAT(IsTrue(bInvoked, TEXT("Callback hasn't been invoked")));
+	}
+
+	TEST_METHOD(AsyncExecute_FExecute_InvokesResultCallbackOnce)
+	{
+		FFakeAsyncTask Task(*TestRunner);
+		int Counter = 0;
+
+		TAsyncExecute<int, FExecute> AsyncExecute = TAsyncExecute<int, FExecute>(
+			*TestRunner,
+			[&]() { return Task.Start(); },
+			FTimespan::FromSeconds(1),
+			nullptr,
+			[&](int) { Counter++; }
+		);
+
+		// Start async task
+		ASSERT_THAT(IsFalse(AsyncExecute.Update(), TEXT("Command stopped execution early")));
+
+		Task.Complete(0);
+		// Need two updates: the first to verify that the async result is ready, the second to invoke the callback.
+		bool bCompleted = AsyncExecute.Update() || AsyncExecute.Update();
+		ASSERT_THAT(IsTrue(bCompleted, TEXT("Command did not stop execution")));
+		ASSERT_THAT(AreEqual(1, Counter, "Callback was invoked multiple times"));
+	}
+
+	TEST_METHOD(AsyncExecute_FExecute_PassesValueToResultCallback)
+	{
+		FFakeAsyncTask Task(*TestRunner);
+		int Result = 0;
+
+		TAsyncExecute<int, FExecute> AsyncExecute = TAsyncExecute<int, FExecute>(
+			*TestRunner,
+			[&]() { return Task.Start(); },
+			FTimespan::FromSeconds(1),
+			nullptr,
+			[&](int InResult) { Result = InResult; }
+		);
+
+		// Start async task
+		ASSERT_THAT(IsFalse(AsyncExecute.Update(), TEXT("Command stopped execution early")));
+
+		const int ExpectedResult = 5;
+		Task.Complete(ExpectedResult);
+		// Need two updates: the first to verify that the async result is ready, the second to invoke the callback.
+		bool bCompleted = AsyncExecute.Update() || AsyncExecute.Update();
+		ASSERT_THAT(IsTrue(bCompleted, TEXT("Command did not stop execution")));
+		ASSERT_THAT(AreEqual(ExpectedResult, Result, TEXT("Incorrect value passed to result callback")));
+	}
+
+	TEST_METHOD(AsyncExecute_FWaitUntil_InvokesResultCallbackOnEveryUpdateWhenAsyncActionIsCompleted)
+	{
+		FFakeAsyncTask Task(*TestRunner);
+		int Counter = 0;
+
+		TAsyncExecute<int, FWaitUntil> AsyncExecute = TAsyncExecute<int, FWaitUntil>(
+			*TestRunner,
+			[&]() { return Task.Start(); },
+			FTimespan::FromSeconds(1),
+			nullptr,
+			[&](int)
+			{
+				Counter++;
+				return false;
+			}
+		);
+
+		// Start async task
+		ASSERT_THAT(IsFalse(AsyncExecute.Update(), TEXT("Command stopped execution early")));
+
+		ASSERT_THAT(AreEqual(0, Counter, TEXT("Callback was invoked early")));
+
+		Task.Complete(0);
+		const int32 Ticks = 10;
+		bool bDone = false;
+		for (int32 i = 0; !bDone && i < Ticks; i++)
+		{
+			bDone = AsyncExecute.Update();
+		}
+		ASSERT_THAT(IsFalse(bDone, TEXT("Command execution stopped before the condition was met")));
+		// One tick will be used to verify that the async result is ready
+		ASSERT_THAT(AreEqual(Ticks - 1, Counter, TEXT("Incorrect number of callback invocations")));
+	}
+
+	TEST_METHOD(AsyncExecute_FWaitUntil_PassesValueToResultCallback)
+	{
+		FFakeAsyncTask Task(*TestRunner);
+		const int ExpectedResult = 6;
+		bool bAllValuesAreValid = true;
+		bool bInvoked = false;
+
+		TAsyncExecute<int, FWaitUntil> AsyncExecute = TAsyncExecute<int, FWaitUntil>(
+			*TestRunner,
+			[&]() { return Task.Start(); },
+			FTimespan::FromSeconds(1),
+			nullptr,
+			[&](int Value)
+			{
+				bInvoked = true;
+				if (Value != ExpectedResult)
+				{
+					bAllValuesAreValid = false;
+				}
+				return false;
+			}
+		);
+
+		// Start async task
+		ASSERT_THAT(IsFalse(AsyncExecute.Update(), TEXT("Command stopped execution early")));
+
+		Task.Complete(ExpectedResult);
+		bool bDone = false;
+		for (int32 i = 0; !bDone && i < 10; i++)
+		{
+			bDone = AsyncExecute.Update();
+		}
+		ASSERT_THAT(IsTrue(bInvoked, TEXT("Result callback was not invoked")));
+		ASSERT_THAT(IsTrue(bAllValuesAreValid, TEXT("Incorrect value passed to result callback")));
+	}
+
+	TEST_METHOD(AsyncExecute_FWaitUntil_CompletesWhenResultCallbackReturnsTrue)
+	{
+		FFakeAsyncTask Task(*TestRunner);
+		bool bCallbackReturnValue = false;
+
+		TAsyncExecute<int, FWaitUntil> AsyncExecute = TAsyncExecute<int, FWaitUntil>(
+			*TestRunner,
+			[&]() { return Task.Start(); },
+			FTimespan::FromSeconds(1),
+			nullptr,
+			[&](int) { return bCallbackReturnValue; }
+		);
+
+		// Start async task
+		ASSERT_THAT(IsFalse(AsyncExecute.Update(), TEXT("Command stopped execution early")));
+
+		Task.Complete(0);
+		bool bDone = false;
+		for (int i = 0; !bDone && i < 10; i++)
+		{
+			bDone = AsyncExecute.Update();
+		}
+		ASSERT_THAT(IsFalse(bDone, TEXT("Command stopped execution before the condition was met")));
+
+		bCallbackReturnValue = true;
+		ASSERT_THAT(IsTrue(AsyncExecute.Update(), TEXT("Command failed to stop execution after the condition was met")));
+	}
+};
+
+TEST_CLASS(AsyncExecuteTimeoutTests, "TestFramework.CQTest.Core")
+{
+	TSharedPtr<FFakeAsyncTask> Task;
+
+	BEFORE_EACH()
+	{
+		Task = MakeShared<FFakeAsyncTask>(*TestRunner);
+	}
+
+	AFTER_EACH()
+	{
+		ASSERT_THAT(IsTrue(Task->InProgress(), TEXT("Async task hasn't been started")));
+		ClearExpectedError(*TestRunner, "Latent command timed out");
+		Task.Reset();
+	}
+
+	TEST_METHOD(AsyncExecute_FExecute_DoesNotInvokeResultCallbackOnTimeout)
+	{
+		AddCommand(MakeShared<TAsyncExecute<int, FExecute>>(
+			*TestRunner,
+			[this]() { return Task->Start(); },
+			FTimespan::FromMilliseconds(1),
+			nullptr,
+			[this](int)
+			{
+				AddError("Result callback should not be invoked after timeout");
+			}
+		));
+	}
+
+	TEST_METHOD(AsyncExecute_FWaitUntil_DoesNotInvokeResultCallbackOnTimeout)
+	{
+		AddCommand(MakeShared<TAsyncExecute<int, FWaitUntil>>(
+			*TestRunner,
+			[this]() { return Task->Start(); },
+			FTimespan::FromMilliseconds(1),
+			nullptr,
+			[this](int)
+			{
+				AddError("Result callback should not be invoked after timeout");
+				return true;
+			}
+		));
+	}
 };
\ No newline at end of file
diff --git a/Engine/Source/Developer/CQTest/Public/Commands/TestCommandBuilder.h b/Engine/Source/Developer/CQTest/Public/Commands/TestCommandBuilder.h
--- a/Engine/Source/Developer/CQTest/Public/Commands/TestCommandBuilder.h
+++ b/Engine/Source/Developer/CQTest/Public/Commands/TestCommandBuilder.h
@@ -55,6 +55,101 @@ public:
 		return Until(nullptr, Query, Timeout);
 	}
 
+	template <typename ResultType>
+	FTestCommandBuilder& DoAsync(const TCHAR* Description, TFunction<TAsyncResult<ResultType>()> AsyncAction, FTimespan Timeout = FTimespan::FromSeconds(10))
+	{
+		if (!TestRunner.HasAnyErrors())
+		{
+			CommandQueue.Add(MakeShared<TAsyncExecute<ResultType>>(TestRunner, MoveTemp(AsyncAction), Timeout, Description));
+		}
+		return *this;
+	}
+
+	template <typename ResultType>
+	FTestCommandBuilder& DoAsync(TFunction<TAsyncResult<ResultType>()> AsyncAction, FTimespan Timeout = FTimespan::FromSeconds(10))
+	{
+		return DoAsync<ResultType>(nullptr, MoveTemp(AsyncAction), Timeout);
+	}
+
+	template <typename ResultType>
+	typename TEnableIf<!std::is_void_v<ResultType>, FTestCommandBuilder&>::Type DoAsync(
+		const TCHAR* Description,
+		TFunction<TAsyncResult<ResultType>()> AsyncAction,
+		TFunction<void(typename TAsyncResultCallbackArg<ResultType>::Type)> ResultCallback,
+		FTimespan Timeout = FTimespan::FromSeconds(10))
+	{
+		if (!TestRunner.HasAnyErrors())
+		{
+			CommandQueue.Add(MakeShared<TAsyncExecute<ResultType, FExecute>>(TestRunner, MoveTemp(AsyncAction), Timeout, Description, MoveTemp(ResultCallback)));
+		}
+		return *this;
+	}
+
+	template <typename ResultType>
+	typename TEnableIf<!std::is_void_v<ResultType>, FTestCommandBuilder&>::Type DoAsync(
+		TFunction<TAsyncResult<ResultType>()> AsyncAction,
+		TFunction<void(typename TAsyncResultCallbackArg<ResultType>::Type)> ResultCallback,
+		FTimespan Timeout = FTimespan::FromSeconds(10))
+	{
+		return DoAsync<ResultType>(nullptr, MoveTemp(AsyncAction), MoveTemp(ResultCallback), Timeout);
+	}
+
+	template <typename ResultType>
+	FTestCommandBuilder& ThenAsync(const TCHAR* Description, TFunction<TAsyncResult<ResultType>()> AsyncAction, FTimespan Timeout = FTimespan::FromSeconds(10))
+	{
+		return DoAsync<ResultType>(Description, MoveTemp(AsyncAction), Timeout);
+	}
+
+	template <typename ResultType>
+	FTestCommandBuilder& ThenAsync(TFunction<TAsyncResult<ResultType>()> AsyncAction, FTimespan Timeout = FTimespan::FromSeconds(10))
+	{
+		return DoAsync<ResultType>(nullptr, MoveTemp(AsyncAction), Timeout);
+	}
+
+	template <typename ResultType>
+	typename TEnableIf<!std::is_void_v<ResultType>, FTestCommandBuilder&>::Type ThenAsync(
+		const TCHAR* Description,
+		TFunction<TAsyncResult<ResultType>()> AsyncAction,
+		TFunction<void(typename TAsyncResultCallbackArg<ResultType>::Type)> ResultCallback,
+		FTimespan Timeout = FTimespan::FromSeconds(10))
+	{
+		return DoAsync<ResultType>(Description, MoveTemp(AsyncAction), MoveTemp(ResultCallback), Timeout);
+	}
+	
+	template <typename ResultType>
+	typename TEnableIf<!std::is_void_v<ResultType>, FTestCommandBuilder&>::Type ThenAsync(
+		TFunction<TAsyncResult<ResultType>()> AsyncAction,
+		TFunction<void(typename TAsyncResultCallbackArg<ResultType>::Type)> ResultCallback,
+		FTimespan Timeout = FTimespan::FromSeconds(10))
+	{
+		return DoAsync<ResultType>(nullptr, MoveTemp(AsyncAction), MoveTemp(ResultCallback), Timeout);
+	}
+
+	template <typename ResultType>
+	typename TEnableIf<!std::is_void_v<ResultType>, FTestCommandBuilder&>::Type UntilAsync(
+		const TCHAR* Description,
+		TFunction<TAsyncResult<ResultType>()> AsyncAction,
+		TFunction<bool(typename TAsyncResultCallbackArg<ResultType>::Type)> ResultCallback,
+		FTimespan AsyncActionTimeout = FTimespan::FromSeconds(10),
+		FTimespan ConditionTimeout = FTimespan::FromSeconds(10))
+	{
+		if (!TestRunner.HasAnyErrors())
+		{
+			CommandQueue.Add(MakeShared<TAsyncExecute<ResultType, FWaitUntil>>(TestRunner, MoveTemp(AsyncAction), AsyncActionTimeout, Description, MoveTemp(ResultCallback), ConditionTimeout));
+		}
+		return *this;
+	}
+
+	template <typename ResultType>
+	typename TEnableIf<!std::is_void_v<ResultType>, FTestCommandBuilder&>::Type UntilAsync(
+		TFunction<TAsyncResult<ResultType>()> AsyncAction,
+		TFunction<bool(typename TAsyncResultCallbackArg<ResultType>::Type)> ResultCallback,
+		FTimespan AsyncActionTimeout = FTimespan::FromSeconds(10),
+		FTimespan ConditionTimeout = FTimespan::FromSeconds(10))
+	{
+		return UntilAsync<ResultType>(nullptr, MoveTemp(AsyncAction), MoveTemp(ResultCallback), AsyncActionTimeout, ConditionTimeout);
+	}
+
 	FTestCommandBuilder& StartWhen(TFunction<bool()> Query, FTimespan Timeout = FTimespan::FromSeconds(10))
 	{
 		return Until(Query, Timeout);
diff --git a/Engine/Source/Developer/CQTest/Public/Commands/TestCommands.h b/Engine/Source/Developer/CQTest/Public/Commands/TestCommands.h
--- a/Engine/Source/Developer/CQTest/Public/Commands/TestCommands.h
+++ b/Engine/Source/Developer/CQTest/Public/Commands/TestCommands.h
@@ -3,6 +3,7 @@
 #pragma once
 
 #include "Misc/AutomationTest.h"
+#include "Async/AsyncResult.h"
 
 /** Latent Command that waits until the Query evaluates to `true` or the timeout has exceeded. */
 class CQTEST_API FWaitUntil : public IAutomationLatentCommand
@@ -10,7 +11,7 @@ class CQTEST_API FWaitUntil : public IAutomationLatentCommand
 public:
 	FWaitUntil(FAutomationTestBase& InTestRunner, TFunction<bool()> Query, FTimespan Timeout = FTimespan::FromSeconds(10), const TCHAR* InDescription = nullptr)
 		: TestRunner(InTestRunner)
-		, Query(MoveTemp(Query)) 
+		, Query(MoveTemp(Query))
 		, Timeout(Timeout)
 		, Description(InDescription)
 	{}
@@ -58,7 +59,7 @@ class CQTEST_API FExecute : public IAutomationLatentCommand
 public:
 	FExecute(FAutomationTestBase& InTestRunner, TFunction<void()> Func, const TCHAR* InDescription = nullptr, ECQTestFailureBehavior InFailureBehavior = ECQTestFailureBehavior::Skip)
 		: TestRunner(InTestRunner)
-		, Func(MoveTemp(Func)) 
+		, Func(MoveTemp(Func))
 		, Description(InDescription)
 		, FailureBehavior(InFailureBehavior)
 	{}
@@ -98,4 +99,122 @@ public:
 	bool Update() override;
 
 	TArray<TSharedPtr<IAutomationLatentCommand>> Commands;
+};
+
+template <typename ResultType>
+struct TAsyncResultCallbackArg
+{
+	using Type = decltype(std::declval<TAsyncResult<ResultType>>().GetFuture().Get());
+};
+
+template <typename ResultType, typename CommandType>
+class TAsyncExecute;
+
+/** This namespace provides a set of variables and functions intended for internal use within the TAsyncExecute class */
+namespace AsyncExecuteDetails
+{
+	static const FTimespan DefaultTimeout = FTimespan::FromSeconds(10);
+
+	static const TCHAR* CreateInternalCommandDescription(TArray<FString>& OutDescriptions, const TCHAR* External, const TCHAR* Internal)
+	{
+		if (!External) { return nullptr; }
+		OutDescriptions.Emplace(FString::Printf(TEXT("%s [%s]"), External, Internal));
+		return *OutDescriptions.Last();
+	}
+
+	template <typename ResultType, typename CommandType>
+	class TResultCommandFactory;
+
+	/** This class provides a helper function for creating a FExecute command to handle the result of an async operation. */
+	template <typename ResultType>
+	class TResultCommandFactory<ResultType, FExecute>
+	{
+		static TSharedRef<IAutomationLatentCommand> Create(
+			FAutomationTestBase& TestRunner,
+			TAsyncResult<ResultType>& AsyncResult,
+			TArray<FString>& OutDescriptions,
+			const TCHAR* BaseDescription,
+			TFunction<void(typename TAsyncResultCallbackArg<ResultType>::Type)>&& InFunc)
+		{
+			return MakeShared<FExecute>(
+				TestRunner,
+				[&Result = AsyncResult, Func = MoveTemp(InFunc)]() { return Func(Result.GetFuture().Get()); },
+				CreateInternalCommandDescription(OutDescriptions, BaseDescription, TEXT("Handle result [FExecute]"))
+			);
+		}
+
+		friend TAsyncExecute<ResultType, FExecute>;
+	};
+	
+	/** This class provides a helper function for creating a FWaitUntil command to handle the result of an async operation. */
+	template <typename ResultType>
+	class TResultCommandFactory<ResultType, FWaitUntil>
+	{
+		static TSharedRef<IAutomationLatentCommand> Create(
+			FAutomationTestBase& TestRunner,
+			TAsyncResult<ResultType>& AsyncResult,
+			TArray<FString>& OutDescriptions,
+			const TCHAR* BaseDescription,
+			TFunction<bool(typename TAsyncResultCallbackArg<ResultType>::Type)>&& InFunc,
+			FTimespan Timeout = DefaultTimeout)
+		{
+			return MakeShared<FWaitUntil>(
+				TestRunner,
+				[&Result = AsyncResult, Func = MoveTemp(InFunc)]() { return Func(Result.GetFuture().Get()); },
+				Timeout,
+				CreateInternalCommandDescription(OutDescriptions, BaseDescription, TEXT("Handle result [FWaitUntil]"))
+			);
+		}
+
+		friend TAsyncExecute<ResultType, FWaitUntil>;
+	};
+};
+
+/** Latent Command that executes an async action and optionally processes the result in a Latent Command of the specified type. */
+template <typename ResultType, typename ResultCommandType = void>
+class TAsyncExecute : public IAutomationLatentCommand
+{
+public:
+	template<typename... ResultCommandArgs>
+	TAsyncExecute(
+		FAutomationTestBase& TestRunner,
+		TFunction<TAsyncResult<ResultType>()> AsyncAction,
+		FTimespan AsyncActionTimeout = AsyncExecuteDetails::DefaultTimeout,
+		const TCHAR* Description = nullptr,
+		ResultCommandArgs&&... Args
+	)
+	{
+		RunSequence.Append(MakeShared<FExecute>(
+			TestRunner,
+			[&Result = AsyncResult, Action = MoveTemp(AsyncAction)]() { Result = Action(); },
+			AsyncExecuteDetails::CreateInternalCommandDescription(Descriptions, Description, TEXT("Execute async action"))
+		));
+		RunSequence.Append(MakeShared<FWaitUntil>(
+			TestRunner,
+			[&Result = AsyncResult]() { return Result.GetFuture().IsReady(); },
+			AsyncActionTimeout,
+			AsyncExecuteDetails::CreateInternalCommandDescription(Descriptions, Description, TEXT("Wait async result"))
+		));
+		if constexpr (!std::is_void_v<ResultCommandType>)
+		{
+			static_assert(!std::is_void_v<ResultType>,
+				"The <ResultType> 'void' should not be used with <ResultCommandType>. "
+				"Instead, consider using a standalone command for any subsequent actions.");
+
+			RunSequence.Append(AsyncExecuteDetails::TResultCommandFactory<ResultType, ResultCommandType>::Create(
+				TestRunner,
+				AsyncResult,
+				Descriptions,
+				Description,
+				std::forward<ResultCommandArgs>(Args)...)
+			);
+		}
+	};
+
+	bool Update() override { return RunSequence.Update(); }
+
+private:
+	FRunSequence RunSequence;
+	TAsyncResult<ResultType> AsyncResult;
+	TArray<FString> Descriptions;
 };
\ No newline at end of file
