Unshelved from pending changelist '37229512':

Fix static analysis warnings in CQTest tests

diff --git a/Engine/Plugins/Tests/CQTest/Source/CQTestTests/Private/Commands/CommandBuilderTests.cpp b/Engine/Plugins/Tests/CQTest/Source/CQTestTests/Private/Commands/CommandBuilderTests.cpp
--- a/Engine/Plugins/Tests/CQTest/Source/CQTestTests/Private/Commands/CommandBuilderTests.cpp
+++ b/Engine/Plugins/Tests/CQTest/Source/CQTestTests/Private/Commands/CommandBuilderTests.cpp
@@ -211,8 +211,12 @@ TEST_CLASS(CommandBuilderTests, "TestFramework.CQTest.Core")
 		const FTestContainer ExpectedResult(567);
 		Promise.SetValue(ExpectedResult);
 		// Need two updates: the first to verify that the async result is ready, the second to invoke the callback.
-		bool bCompleted = Command->Update() || Command->Update();
-		ASSERT_THAT(IsTrue(bCompleted));
+		bool bDone = false;
+		for (int32 i = 0; !bDone && i < 2; i++)
+		{
+			bDone = Command->Update();
+		}
+		ASSERT_THAT(IsTrue(bDone));
 		ASSERT_THAT(AreEqual(ExpectedResult.Value, Result.Value));
 	}
 
@@ -238,8 +242,12 @@ TEST_CLASS(CommandBuilderTests, "TestFramework.CQTest.Core")
 
 		Promise.SetValue(Value);
 		// Need two updates: the first to verify that the async result is ready, the second to invoke the callback.
-		bool bCompleted = Command->Update() || Command->Update();
-		ASSERT_THAT(IsTrue(bCompleted));
+		bool bDone = false;
+		for (int32 i = 0; !bDone && i < 2; i++)
+		{
+			bDone = Command->Update();
+		}
+		ASSERT_THAT(IsTrue(bDone));
 		ASSERT_THAT(AreEqual(ExpectedValue, Value));
 	}
 
@@ -375,8 +383,12 @@ TEST_CLASS(CommandBuilderTests, "TestFramework.CQTest.Core")
 		const FTestContainer ExpectedResult(567);
 		Promise.SetValue(ExpectedResult);
 		// Need two updates: the first to verify that the async result is ready, the second to invoke the callback.
-		bool bCompleted = Command->Update() || Command->Update();
-		ASSERT_THAT(IsTrue(bCompleted));
+		bool bDone = false;
+		for (int32 i = 0; !bDone && i < 2; i++)
+		{
+			bDone = Command->Update();
+		}
+		ASSERT_THAT(IsTrue(bDone));
 		ASSERT_THAT(AreEqual(ExpectedResult.Value, Result.Value));
 	}
 
@@ -403,8 +415,12 @@ TEST_CLASS(CommandBuilderTests, "TestFramework.CQTest.Core")
 
 		Promise.SetValue(Value);
 		// Need two updates: the first to verify that the async result is ready, the second to invoke the callback.
-		bool bCompleted = Command->Update() || Command->Update();
-		ASSERT_THAT(IsTrue(bCompleted));
+		bool bDone = false;
+		for (int32 i = 0; !bDone && i < 2; i++)
+		{
+			bDone = Command->Update();
+		}
+		ASSERT_THAT(IsTrue(bDone));
 		ASSERT_THAT(AreEqual(ExpectedValue, Value));
 	}
 	
diff --git a/Engine/Plugins/Tests/CQTest/Source/CQTestTests/Private/Commands/TestCommandsTests.cpp b/Engine/Plugins/Tests/CQTest/Source/CQTestTests/Private/Commands/TestCommandsTests.cpp
--- a/Engine/Plugins/Tests/CQTest/Source/CQTestTests/Private/Commands/TestCommandsTests.cpp
+++ b/Engine/Plugins/Tests/CQTest/Source/CQTestTests/Private/Commands/TestCommandsTests.cpp
@@ -421,8 +421,12 @@ TEST_CLASS(AsyncExecuteBasicTests, "TestFramework.CQTest.Core")
 
 		Task.Complete(0);
 		// Need two updates: the first to verify that the async result is ready, the second to invoke the callback.
-		bool bCompleted = AsyncExecute.Update() || AsyncExecute.Update();
-		ASSERT_THAT(IsTrue(bCompleted, TEXT("Command failed to stop execution after the async task was completed")));
+		bool bDone = false;
+		for (int32 i = 0; !bDone && i < 2; i++)
+		{
+			bDone = AsyncExecute.Update();
+		}
+		ASSERT_THAT(IsTrue(bDone, TEXT("Command failed to stop execution after the async task was completed")));
 	}
 
 	TEST_METHOD(AsyncExecute_FExecute_InvokesResultCallbackWhenAsyncActionIsCompleted)
@@ -445,8 +449,12 @@ TEST_CLASS(AsyncExecuteBasicTests, "TestFramework.CQTest.Core")
 		Task.Complete(0);
 
 		// Need two updates: the first to verify that the async result is ready, the second to invoke the callback.
-		bool bCompleted = AsyncExecute.Update() || AsyncExecute.Update();
-		ASSERT_THAT(IsTrue(bCompleted, TEXT("Command did not stop execution")));
+		bool bDone = false;
+		for (int32 i = 0; !bDone && i < 2; i++)
+		{
+			bDone = AsyncExecute.Update();
+		}
+		ASSERT_THAT(IsTrue(bDone, TEXT("Command did not stop execution")));
 		ASSERT_THAT(IsTrue(bInvoked, TEXT("Callback hasn't been invoked")));
 	}
 
@@ -468,8 +476,12 @@ TEST_CLASS(AsyncExecuteBasicTests, "TestFramework.CQTest.Core")
 
 		Task.Complete(0);
 		// Need two updates: the first to verify that the async result is ready, the second to invoke the callback.
-		bool bCompleted = AsyncExecute.Update() || AsyncExecute.Update();
-		ASSERT_THAT(IsTrue(bCompleted, TEXT("Command did not stop execution")));
+		bool bDone = false;
+		for (int32 i = 0; !bDone && i < 2; i++)
+		{
+			bDone = AsyncExecute.Update();
+		}
+		ASSERT_THAT(IsTrue(bDone, TEXT("Command did not stop execution")));
 		ASSERT_THAT(AreEqual(1, Counter, "Callback was invoked multiple times"));
 	}
 
@@ -492,8 +504,12 @@ TEST_CLASS(AsyncExecuteBasicTests, "TestFramework.CQTest.Core")
 		const int ExpectedResult = 5;
 		Task.Complete(ExpectedResult);
 		// Need two updates: the first to verify that the async result is ready, the second to invoke the callback.
-		bool bCompleted = AsyncExecute.Update() || AsyncExecute.Update();
-		ASSERT_THAT(IsTrue(bCompleted, TEXT("Command did not stop execution")));
+		bool bDone = false;
+		for (int32 i = 0; !bDone && i < 2; i++)
+		{
+			bDone = AsyncExecute.Update();
+		}
+		ASSERT_THAT(IsTrue(bDone, TEXT("Command did not stop execution")));
 		ASSERT_THAT(AreEqual(ExpectedResult, Result, TEXT("Incorrect value passed to result callback")));
 	}
 
