Adjusting timeout for the async tests to take into account slower platforms

#rnx

diff --git a/Engine/Plugins/Tests/CQTest/Source/CQTestTests/Private/Commands/CommandBuilderTests.cpp b/Engine/Plugins/Tests/CQTest/Source/CQTestTests/Private/Commands/CommandBuilderTests.cpp
--- a/Engine/Plugins/Tests/CQTest/Source/CQTestTests/Private/Commands/CommandBuilderTests.cpp
+++ b/Engine/Plugins/Tests/CQTest/Source/CQTestTests/Private/Commands/CommandBuilderTests.cpp
@@ -298,8 +298,6 @@ TEST_CLASS(CommandBuilderTests, "TestFramework.CQTest.Core")
 		TPromise<bool> Promise;
 		const FTimespan AsyncActionTimeout = FTimespan::FromMilliseconds(200);
 		const FTimespan ConditionTimeout = FTimespan::FromSeconds(10);
-		const FDateTime StartTime = FDateTime::UtcNow();
-		const FDateTime MaxEndTime = StartTime + FTimespan::FromMilliseconds(500);
 
 		auto Command = CommandBuilder.UntilAsync<bool>(
 			[&]() { return TAsyncResult<bool>(Promise.GetFuture(), nullptr, nullptr); },
@@ -309,6 +307,8 @@ TEST_CLASS(CommandBuilderTests, "TestFramework.CQTest.Core")
 		).Build();
 
 		bool bDone = false;
+		const FDateTime StartTime = FDateTime::UtcNow();
+		const FDateTime MaxEndTime = StartTime + FTimespan::FromMilliseconds(5000);
 		while (!bDone && FDateTime::UtcNow() < MaxEndTime)
 		{
 			bDone = Command->Update();
@@ -328,8 +328,6 @@ TEST_CLASS(CommandBuilderTests, "TestFramework.CQTest.Core")
 		TPromise<bool> Promise;
 		const FTimespan AsyncActionTimeout = FTimespan::FromSeconds(10);
 		const FTimespan ConditionTimeout = FTimespan::FromMilliseconds(200);
-		const FDateTime StartTime = FDateTime::UtcNow();
-		const FDateTime MaxEndTime = StartTime + FTimespan::FromMilliseconds(500);
 
 		auto Command = CommandBuilder.UntilAsync<bool>(
 			[&]() { return TAsyncResult<bool>(Promise.GetFuture(), nullptr, nullptr); },
@@ -341,6 +339,8 @@ TEST_CLASS(CommandBuilderTests, "TestFramework.CQTest.Core")
 		Promise.SetValue(true);
 
 		bool bDone = false;
+		const FDateTime StartTime = FDateTime::UtcNow();
+		const FDateTime MaxEndTime = StartTime + FTimespan::FromMilliseconds(5000);
 		while (!bDone && FDateTime::UtcNow() < MaxEndTime)
 		{
 			bDone = Command->Update();
