CQTest - Support assert macros in base test classes

PR #13128

diff --git a/Engine/Plugins/Tests/CQTest/Source/CQTestTests/Private/ExtensionTests.cpp b/Engine/Plugins/Tests/CQTest/Source/CQTestTests/Private/ExtensionTests.cpp
--- a/Engine/Plugins/Tests/CQTest/Source/CQTestTests/Private/ExtensionTests.cpp
+++ b/Engine/Plugins/Tests/CQTest/Source/CQTestTests/Private/ExtensionTests.cpp
@@ -64,6 +64,11 @@ struct TCustomBaseClass : public TTest<Derived, AsserterType>
 	{
 		BaseValue = 0;
 	}
+
+	void AssertInBase()
+	{
+		ASSERT_THAT(AreEqual(42, SpecialValue));
+	}
 };
 
 #define CUSTOM_BASE_TEST_CLASS(_ClassName, _TestDir) TEST_CLASS_WITH_BASE(_ClassName, _TestDir, TCustomBaseClass)
@@ -94,6 +99,11 @@ CUSTOM_BASE_TEST_CLASS(DerivedTest, "TestFramework.CQTest.Core")
 		ASSERT_THAT(AreEqual(42, BaseValue));
 		ASSERT_THAT(AreEqual(BaseValue, DerivedValue));
 	}
+
+	TEST_METHOD(DerivedTestClass_WithFunctionThatAssertsInBase_CanAssertInBase)
+	{
+		AssertInBase();
+	}
 };
 
 TEST_CLASS_WITH_BASE_AND_FLAGS(DerivedWithFlagsTest, "TestFramework.CQTest.Core", TCustomBaseClass, EAutomationTestFlags_ApplicationContextMask | EAutomationTestFlags::EngineFilter)
diff --git a/Engine/Source/Developer/CQTest/Public/CQTest.h b/Engine/Source/Developer/CQTest/Public/CQTest.h
--- a/Engine/Source/Developer/CQTest/Public/CQTest.h
+++ b/Engine/Source/Developer/CQTest/Public/CQTest.h
@@ -321,8 +321,8 @@ struct TTest : TBaseTest<AsserterType>
 	};                                                   \
 	void _TestName::_TestName##_Method()
 
-#define ASSERT_THAT(_assertion) if (!Assert._assertion) {return;}
-#define ASSERT_FAIL(_Msg) Assert.Fail(_Msg); return;
+#define ASSERT_THAT(_assertion) if (!this->Assert._assertion) {return;}
+#define ASSERT_FAIL(_Msg) this->Assert.Fail(_Msg); return;
 
 #define BEFORE_ALL() static void BeforeAll(const FString&)
 #define BEFORE_EACH() virtual void Setup() override
