Fix for null entries in ClientConnections array

#rnx
#13605

diff --git a/Engine/Plugins/Tests/CQTest/Source/CQTestTests/Private/Components/PIENetworkComponentTests.cpp b/Engine/Plugins/Tests/CQTest/Source/CQTestTests/Private/Components/PIENetworkComponentTests.cpp
--- a/Engine/Plugins/Tests/CQTest/Source/CQTestTests/Private/Components/PIENetworkComponentTests.cpp
+++ b/Engine/Plugins/Tests/CQTest/Source/CQTestTests/Private/Components/PIENetworkComponentTests.cpp
@@ -98,6 +98,19 @@ NETWORK_TEST_CLASS(StateTest, "TestFramework.CQTest.Network")
 			.UntilClients([this](DerivedState& State) { SharedNumber++;  return ++State.IndependentNumber > 4; })
 			.Then([this]() { ASSERT_THAT(AreEqual(15, SharedNumber)); });
 	}
+
+	TEST_METHOD(Network_WithClients_ClientConnectionsAreFilledIn)
+	{
+		Network
+			.ThenServer([this](DerivedState& State) {
+				ASSERT_THAT(AreEqual(State.ClientCount, 3));
+				ASSERT_THAT(AreEqual(State.ClientConnections.Num(), 3));
+				for (int32 ClientIndex = 0; ClientIndex < State.ClientConnections.Num(); ++ ClientIndex)
+				{
+					ASSERT_THAT(IsNotNull(State.ClientConnections[ClientIndex]));
+				}
+			});
+	}
 };
 
 NETWORK_TEST_CLASS(ReplicationTest, "TestFramework.CQTest.Network")
@@ -297,6 +310,19 @@ NETWORK_TEST_CLASS(LateJoinTest, "TestFramework.CQTest.Network")
 		.ThenServer([this](DerivedState& ServerState) {
 			ASSERT_THAT(AreEqual(ServerState.ClientCount, 2));
 			ASSERT_THAT(AreEqual(ServerState.ClientConnections.Num(), 2));
+			for (int32 ClientIndex = 0; ClientIndex < ServerState.ClientConnections.Num(); ++ClientIndex)
+			{
+				ASSERT_THAT(IsNotNull(ServerState.ClientConnections[ClientIndex]));
+			}
+		})
+		.ThenClientJoins()
+		.ThenServer([this](DerivedState& ServerState) {
+			ASSERT_THAT(AreEqual(ServerState.ClientCount, 3));
+			ASSERT_THAT(AreEqual(ServerState.ClientConnections.Num(), 3));
+			for (int32 ClientIndex = 0; ClientIndex < ServerState.ClientConnections.Num(); ++ClientIndex)
+			{
+				ASSERT_THAT(IsNotNull(ServerState.ClientConnections[ClientIndex]));
+			}
 		});
 	}
 
diff --git a/Engine/Plugins/Tests/CQTestExperimental/Source/CQTestExperimentalTests/Private/Iris/PIENetworkComponentTests.cpp b/Engine/Plugins/Tests/CQTestExperimental/Source/CQTestExperimentalTests/Private/Iris/PIENetworkComponentTests.cpp
--- a/Engine/Plugins/Tests/CQTestExperimental/Source/CQTestExperimentalTests/Private/Iris/PIENetworkComponentTests.cpp
+++ b/Engine/Plugins/Tests/CQTestExperimental/Source/CQTestExperimentalTests/Private/Iris/PIENetworkComponentTests.cpp
@@ -95,6 +95,19 @@ NETWORK_TEST_CLASS(IrisStateTest, "TestFramework.CQTest.Experimental.Network")
 			.UntilClients([this](DerivedState& State) { SharedNumber++;  return ++State.IndependentNumber > 4; })
 			.Then([this]() { ASSERT_THAT(AreEqual(15, SharedNumber)); });
 	}
+
+	TEST_METHOD(Network_WithClients_ClientConnectionsAreFilledIn)
+	{
+		Network
+			.ThenServer([this](DerivedState& State) {
+				ASSERT_THAT(AreEqual(State.ClientCount, 3));
+				ASSERT_THAT(AreEqual(State.ClientConnections.Num(), 3));
+				for (int32 ClientIndex = 0; ClientIndex < State.ClientConnections.Num(); ++ ClientIndex)
+				{
+					ASSERT_THAT(IsNotNull(State.ClientConnections[ClientIndex]));
+				}
+			});
+	}
 };
 
 NETWORK_TEST_CLASS(IrisReplicationTest, "TestFramework.CQTest.Experimental.Network")
@@ -294,6 +307,10 @@ NETWORK_TEST_CLASS(IrisLateJoinTest, "TestFramework.CQTest.Experimental.Network"
 		.ThenServer([this](DerivedState& ServerState) {
 			ASSERT_THAT(AreEqual(ServerState.ClientCount, 2));
 			ASSERT_THAT(AreEqual(ServerState.ClientConnections.Num(), 2));
+			for (int32 ClientIndex = 0; ClientIndex < ServerState.ClientConnections.Num(); ++ClientIndex)
+			{
+				ASSERT_THAT(IsNotNull(ServerState.ClientConnections[ClientIndex]));
+			}
 		});
 	}
 
diff --git a/Engine/Source/Developer/CQTest/Private/Components/PIENetworkComponent.cpp b/Engine/Source/Developer/CQTest/Private/Components/PIENetworkComponent.cpp
--- a/Engine/Source/Developer/CQTest/Private/Components/PIENetworkComponent.cpp
+++ b/Engine/Source/Developer/CQTest/Private/Components/PIENetworkComponent.cpp
@@ -224,10 +224,10 @@ void FBasePIENetworkComponent::SetPacketSettings() const
 	}
 }
 
-void FBasePIENetworkComponent::ConnectClientsToServer() 
+void FBasePIENetworkComponent::ConnectClientsToServer()
 {
 	auto& ServerConnections = ServerState->World->GetNetDriver()->ClientConnections;
-	for(int32 ClientIndex = ServerState->ClientConnections.Num(); ClientIndex < ServerState->ClientCount; ClientIndex++)
+	for(int32 ClientIndex = 0; ClientIndex < ServerState->ClientCount; ClientIndex++)
 	{
 		const int32 ClientLocalPort = ClientStates[ClientIndex]->World->GetNetDriver()->GetLocalAddr()->GetPort();
 		TObjectPtr<UNetConnection>* ServerConnection = ServerConnections.FindByPredicate([ClientLocalPort](UNetConnection* ClientConnection) {
